[BugFix] Inconsistent behavior in root component transform scale values for dynamically-spawned Blueprint actor instances.

diff --git a/Engine/Source/Editor/BlueprintGraph/Classes/K2Node_SpawnActorFromClass.h b/Engine/Source/Editor/BlueprintGraph/Classes/K2Node_SpawnActorFromClass.h
--- a/Engine/Source/Editor/BlueprintGraph/Classes/K2Node_SpawnActorFromClass.h
+++ b/Engine/Source/Editor/BlueprintGraph/Classes/K2Node_SpawnActorFromClass.h
@@ -39,8 +39,14 @@ class BLUEPRINTGRAPH_API UK2Node_SpawnActorFromClass : public UK2Node_ConstructO
 	virtual void ExpandNode(class FKismetCompilerContext& CompilerContext, UEdGraph* SourceGraph) override;
 	virtual FSlateIcon GetIconAndTint(FLinearColor& OutColor) const override;
 	virtual bool IsCompatibleWithGraph(const UEdGraph* TargetGraph) const override;
+	virtual void PostPlacedNewNode() override;
 	//~ End UEdGraphNode Interface.
 
+	//~ Begin UObject Interface
+	virtual void Serialize(FArchive& Ar) override;
+	virtual void PostLoad() override;
+	//~ End UObject Interface
+	
 	//~ Begin UK2Node Interface
 	virtual bool IsNodeSafeToIgnore() const override { return true; }
 	virtual void ReallocatePinsDuringReconstruction(TArray<UEdGraphPin*>& OldPins) override;
@@ -53,11 +59,16 @@ class BLUEPRINTGRAPH_API UK2Node_SpawnActorFromClass : public UK2Node_ConstructO
 	virtual bool IsSpawnVarPin(UEdGraphPin* Pin) const override;
 	//~ End UK2Node_ConstructObjectFromClass Interface
 
+	
 private:
+	void FixupScaleMethodPin();
+	
 	/** Get the spawn transform input pin */	
 	UEdGraphPin* GetSpawnTransformPin() const;
 	/** Get the collision handling method input pin */
 	UEdGraphPin* GetCollisionHandlingOverridePin() const;
+	/** Get the collision handling method input pin */
+	UEdGraphPin* GetScaleMethodPin() const;
 	/** Get the actor owner pin */
 	UEdGraphPin* GetOwnerPin() const;
 
diff --git a/Engine/Source/Editor/BlueprintGraph/Private/K2Node_SpawnActorFromClass.cpp b/Engine/Source/Editor/BlueprintGraph/Private/K2Node_SpawnActorFromClass.cpp
--- a/Engine/Source/Editor/BlueprintGraph/Private/K2Node_SpawnActorFromClass.cpp
+++ b/Engine/Source/Editor/BlueprintGraph/Private/K2Node_SpawnActorFromClass.cpp
@@ -35,6 +35,7 @@
 #include "UObject/NameTypes.h"
 #include "UObject/Object.h"
 #include "UObject/ObjectPtr.h"
+#include "UObject/UE5MainStreamObjectVersion.h"
 
 struct FLinearColor;
 
@@ -44,6 +45,7 @@ struct FK2Node_SpawnActorFromClassHelper
 	static const FName SpawnEvenIfCollidingPinName;
 	static const FName NoCollisionFailPinName;
 	static const FName CollisionHandlingOverridePinName;
+	static const FName TransformScaleMethodPinName;
 	static const FName OwnerPinName;
 };
 
@@ -51,6 +53,7 @@ const FName FK2Node_SpawnActorFromClassHelper::SpawnTransformPinName(TEXT("Spawn
 const FName FK2Node_SpawnActorFromClassHelper::SpawnEvenIfCollidingPinName(TEXT("SpawnEvenIfColliding"));		// deprecated pin, name kept for backwards compat
 const FName FK2Node_SpawnActorFromClassHelper::NoCollisionFailPinName(TEXT("bNoCollisionFail"));		// deprecated pin, name kept for backwards compat
 const FName FK2Node_SpawnActorFromClassHelper::CollisionHandlingOverridePinName(TEXT("CollisionHandlingOverride"));
+const FName FK2Node_SpawnActorFromClassHelper::TransformScaleMethodPinName(TEXT("TransformScaleMethod"));
 const FName FK2Node_SpawnActorFromClassHelper::OwnerPinName(TEXT("Owner"));
 
 #define LOCTEXT_NAMESPACE "K2Node_SpawnActorFromClass"
@@ -75,10 +78,16 @@ void UK2Node_SpawnActorFromClass::AllocateDefaultPins()
 	UEdGraphPin* TransformPin = CreatePin(EGPD_Input, UEdGraphSchema_K2::PC_Struct, TransformStruct, FK2Node_SpawnActorFromClassHelper::SpawnTransformPinName);
 
 	// Collision handling method pin
-	UEnum* const MethodEnum = FindObjectChecked<UEnum>(nullptr, TEXT("/Script/Engine.ESpawnActorCollisionHandlingMethod"), true);
-	UEdGraphPin* const CollisionHandlingOverridePin = CreatePin(EGPD_Input, UEdGraphSchema_K2::PC_Byte, MethodEnum, FK2Node_SpawnActorFromClassHelper::CollisionHandlingOverridePinName);
-	CollisionHandlingOverridePin->DefaultValue = MethodEnum->GetNameStringByValue(static_cast<int>(ESpawnActorCollisionHandlingMethod::Undefined));
-
+	UEnum* const CollisionMethodEnum = StaticEnum<ESpawnActorCollisionHandlingMethod>();
+	UEdGraphPin* const CollisionHandlingOverridePin = CreatePin(EGPD_Input, UEdGraphSchema_K2::PC_Byte, CollisionMethodEnum, FK2Node_SpawnActorFromClassHelper::CollisionHandlingOverridePinName);
+	CollisionHandlingOverridePin->DefaultValue = CollisionMethodEnum->GetNameStringByValue(static_cast<int>(ESpawnActorCollisionHandlingMethod::Undefined));
+	
+	// Pin to set transform scaling behavior (ie whether to multiply scale with root component or to just ignore the root component default scale)
+	UEnum* const ScaleMethodEnum = StaticEnum<ESpawnActorScaleMethod>();
+	UEdGraphPin* const ScaleMethodPin = CreatePin(EGPD_Input, UEdGraphSchema_K2::PC_Byte, ScaleMethodEnum, FK2Node_SpawnActorFromClassHelper::TransformScaleMethodPinName);
+	ScaleMethodPin->DefaultValue = ScaleMethodEnum->GetNameStringByValue(static_cast<int>(ESpawnActorScaleMethod::MultiplyWithRoot));
+	ScaleMethodPin->bAdvancedView = true;
+	
 	UEdGraphPin* OwnerPin = CreatePin(EGPD_Input, UEdGraphSchema_K2::PC_Object, AActor::StaticClass(), FK2Node_SpawnActorFromClassHelper::OwnerPinName);
 	OwnerPin->bAdvancedView = true;
 	if (ENodeAdvancedPins::NoPins == AdvancedPinDisplay)
@@ -228,6 +237,7 @@ bool UK2Node_SpawnActorFromClass::IsSpawnVarPin(UEdGraphPin* Pin) const
 	}
 
 	return(	Super::IsSpawnVarPin(Pin) &&
+			Pin->PinName != FK2Node_SpawnActorFromClassHelper::TransformScaleMethodPinName &&
 			Pin->PinName != FK2Node_SpawnActorFromClassHelper::CollisionHandlingOverridePinName &&
 			Pin->PinName != FK2Node_SpawnActorFromClassHelper::SpawnTransformPinName && 
 			Pin->PinName != FK2Node_SpawnActorFromClassHelper::OwnerPinName );
@@ -273,6 +283,13 @@ UEdGraphPin* UK2Node_SpawnActorFromClass::GetCollisionHandlingOverridePin() cons
 	return Pin;
 }
 
+UEdGraphPin* UK2Node_SpawnActorFromClass::GetScaleMethodPin() const
+{
+	UEdGraphPin* const Pin = FindPinChecked(FK2Node_SpawnActorFromClassHelper::TransformScaleMethodPinName);
+	check(Pin->Direction == EGPD_Input);
+	return Pin;
+}
+
 UEdGraphPin* UK2Node_SpawnActorFromClass::GetOwnerPin() const
 {
 	UEdGraphPin* Pin = FindPin(FK2Node_SpawnActorFromClassHelper::OwnerPinName);
@@ -329,6 +346,47 @@ bool UK2Node_SpawnActorFromClass::IsCompatibleWithGraph(const UEdGraph* TargetGr
 	return Super::IsCompatibleWithGraph(TargetGraph) && (!Blueprint || (FBlueprintEditorUtils::FindUserConstructionScript(Blueprint) != TargetGraph && Blueprint->GeneratedClass->GetDefaultObject()->ImplementsGetWorld()));
 }
 
+void UK2Node_SpawnActorFromClass::PostPlacedNewNode()
+{
+	UEdGraphPin* const ScaleMethodPin = GetScaleMethodPin();
+	const UEnum* const ScaleMethodEnum = StaticEnum<ESpawnActorScaleMethod>();
+	ScaleMethodPin->DefaultValue = ScaleMethodEnum->GetNameStringByValue(static_cast<int>(ESpawnActorScaleMethod::MultiplyWithRoot));
+}
+
+void UK2Node_SpawnActorFromClass::Serialize(FArchive& Ar)
+{
+	Super::Serialize(Ar);
+	Ar.UsingCustomVersion(FUE5MainStreamObjectVersion::GUID);
+}
+
+void UK2Node_SpawnActorFromClass::PostLoad()
+{
+	FixupScaleMethodPin();
+	Super::PostLoad();
+}
+
+void UK2Node_SpawnActorFromClass::FixupScaleMethodPin()
+{
+	if (GetLinkerCustomVersion(FUE5MainStreamObjectVersion::GUID) < FUE5MainStreamObjectVersion::SpawnActorFromClassTransformScaleMethod)
+	{
+		const UEdGraphPin* const ClassPin = FindPin(TEXT( "Class" ));
+		const UClass* Class = Cast<UClass>(ClassPin->DefaultObject);
+		if (const AActor* ActorCDO = Cast<AActor>(Class->ClassDefaultObject))
+		{
+			UEdGraphPin* const ScaleMethodPin = GetScaleMethodPin();
+			const UEnum* const ScaleMethodEnum = StaticEnum<ESpawnActorScaleMethod>();
+			if (const USceneComponent* Root = ActorCDO->GetRootComponent()) // native root component
+			{
+				ScaleMethodPin->DefaultValue = ScaleMethodEnum->GetNameStringByValue(static_cast<int>(ESpawnActorScaleMethod::MultiplyWithRoot));
+			}
+			else
+			{
+				ScaleMethodPin->DefaultValue = ScaleMethodEnum->GetNameStringByValue(static_cast<int>(ESpawnActorScaleMethod::OverrideRootScale));
+			}
+		}
+	}
+}
+
 void UK2Node_SpawnActorFromClass::GetNodeAttributes( TArray<TKeyValuePair<FString, FString>>& OutNodeAttributes ) const
 {
 	UClass* ClassToSpawn = GetClassToSpawn();
@@ -348,6 +406,9 @@ void UK2Node_SpawnActorFromClass::ExpandNode(class FKismetCompilerContext& Compi
 {
 	Super::ExpandNode(CompilerContext, SourceGraph);
 
+	// make sure that if scale method needs to be given a default value, it gets it
+	FixupScaleMethodPin();
+
 	static const FName BeginSpawningBlueprintFuncName = GET_FUNCTION_NAME_CHECKED(UGameplayStatics, BeginDeferredActorSpawnFromClass);
 	static const FName ActorClassParamName(TEXT("ActorClass"));
 	static const FName WorldContextParamName(TEXT("WorldContextObject"));
@@ -366,6 +427,7 @@ void UK2Node_SpawnActorFromClass::ExpandNode(class FKismetCompilerContext& Compi
 	UEdGraphPin* SpawnNodeExec = SpawnNode->GetExecPin();
 	UEdGraphPin* SpawnNodeTransform = SpawnNode->GetSpawnTransformPin();
 	UEdGraphPin* SpawnNodeCollisionHandlingOverride = GetCollisionHandlingOverridePin();
+	UEdGraphPin* SpawnNodeScaleMethod = GetScaleMethodPin();
 	UEdGraphPin* SpawnWorldContextPin = SpawnNode->GetWorldContextPin();
 	UEdGraphPin* SpawnClassPin = SpawnNode->GetClassPin();
 	UEdGraphPin* SpawnNodeOwnerPin = SpawnNode->GetOwnerPin();
@@ -394,6 +456,7 @@ void UK2Node_SpawnActorFromClass::ExpandNode(class FKismetCompilerContext& Compi
 	UEdGraphPin* CallBeginWorldContextPin = CallBeginSpawnNode->FindPinChecked(WorldContextParamName);
 	UEdGraphPin* CallBeginActorClassPin = CallBeginSpawnNode->FindPinChecked(ActorClassParamName);
 	UEdGraphPin* CallBeginTransform = CallBeginSpawnNode->FindPinChecked(TransformParamName);
+	UEdGraphPin* CallBeginScaleMethod = CallBeginSpawnNode->FindPinChecked(FK2Node_SpawnActorFromClassHelper::TransformScaleMethodPinName);
 	UEdGraphPin* CallBeginCollisionHandlingOverride = CallBeginSpawnNode->FindPinChecked(CollisionHandlingOverrideParamName);
 
 	UEdGraphPin* CallBeginOwnerPin = CallBeginSpawnNode->FindPinChecked(FK2Node_SpawnActorFromClassHelper::OwnerPinName);
@@ -427,9 +490,12 @@ void UK2Node_SpawnActorFromClass::ExpandNode(class FKismetCompilerContext& Compi
 	// Copy the 'transform' connection from the spawn node to 'begin spawn'
 	CompilerContext.MovePinLinksToIntermediate(*SpawnNodeTransform, *CallBeginTransform);
 
-	// Copy the 'bNoCollisionFail' connection from the spawn node to 'begin spawn'
+	// Copy the 'CollisionHandlingOverride' connection from the spawn node to 'begin spawn'
 	CompilerContext.MovePinLinksToIntermediate(*SpawnNodeCollisionHandlingOverride, *CallBeginCollisionHandlingOverride);
 
+	// Copy the 'ScaleMethod' connection from the spawn node to 'begin spawn'
+	CompilerContext.MovePinLinksToIntermediate(*SpawnNodeScaleMethod, *CallBeginScaleMethod);
+
 	//////////////////////////////////////////////////////////////////////////
 	// create 'finish spawn' call node
 	UK2Node_CallFunction* CallFinishSpawnNode = CompilerContext.SpawnIntermediateNode<UK2Node_CallFunction>(SpawnNode, SourceGraph);
@@ -440,6 +506,7 @@ void UK2Node_SpawnActorFromClass::ExpandNode(class FKismetCompilerContext& Compi
 	UEdGraphPin* CallFinishThen = CallFinishSpawnNode->GetThenPin();
 	UEdGraphPin* CallFinishActor = CallFinishSpawnNode->FindPinChecked(ActorParamName);
 	UEdGraphPin* CallFinishTransform = CallFinishSpawnNode->FindPinChecked(TransformParamName);
+	UEdGraphPin* CallFinishScaleMethod = CallFinishSpawnNode->FindPinChecked(FK2Node_SpawnActorFromClassHelper::TransformScaleMethodPinName);
 	UEdGraphPin* CallFinishResult = CallFinishSpawnNode->GetReturnValuePin();
 
 	// Move 'then' connection from spawn node to 'finish spawn'
@@ -447,6 +514,9 @@ void UK2Node_SpawnActorFromClass::ExpandNode(class FKismetCompilerContext& Compi
 
 	// Copy transform connection
 	CompilerContext.CopyPinLinksToIntermediate(*CallBeginTransform, *CallFinishTransform);
+ 
+	// Copy transform connection
+	CompilerContext.CopyPinLinksToIntermediate(*CallBeginScaleMethod, *CallFinishScaleMethod);
 
 	// Connect output actor from 'begin' to 'finish'
 	CallBeginResult->MakeLinkTo(CallFinishActor);
diff --git a/Engine/Source/Runtime/Core/Public/UObject/UE5MainStreamObjectVersion.h b/Engine/Source/Runtime/Core/Public/UObject/UE5MainStreamObjectVersion.h
--- a/Engine/Source/Runtime/Core/Public/UObject/UE5MainStreamObjectVersion.h
+++ b/Engine/Source/Runtime/Core/Public/UObject/UE5MainStreamObjectVersion.h
@@ -292,6 +292,9 @@ struct CORE_API FUE5MainStreamObjectVersion
 		// bIsUsedWithVolumetricCloud flag auto conversion
 		UpdateHairDescriptionBulkData,
 
+		// Added TransformScaleMethod pin to SpawnActorFromClass node
+		SpawnActorFromClassTransformScaleMethod,
+
 		// -----<new versions can be added above this line>-------------------------------------------------
 		VersionPlusOne,
 		LatestVersion = VersionPlusOne - 1
diff --git a/Engine/Source/Runtime/Engine/Classes/Engine/SCS_Node.h b/Engine/Source/Runtime/Engine/Classes/Engine/SCS_Node.h
--- a/Engine/Source/Runtime/Engine/Classes/Engine/SCS_Node.h
+++ b/Engine/Source/Runtime/Engine/Classes/Engine/SCS_Node.h
@@ -108,7 +108,7 @@ class USCS_Node : public UObject
 	 * @param bIsDefaultTransform	Indicates whether or not the given transform is a "default" transform, in which case it can be overridden by template defaults.
 	 * @return The new component instance that was created, or NULL on failure.
 	 */
-	UActorComponent* ExecuteNodeOnActor(AActor* Actor, USceneComponent* ParentComponent, const FTransform* RootTransform, const struct FRotationConversionCache* RootRelativeRotationCache, bool bIsDefaultTransform);
+	UActorComponent* ExecuteNodeOnActor(AActor* Actor, USceneComponent* ParentComponent, const FTransform* RootTransform, const struct FRotationConversionCache* RootRelativeRotationCache, bool bIsDefaultTransform, ESpawnActorScaleMethod TransformScaleMethod = ESpawnActorScaleMethod::OverrideRootScale);
 
 	/** Return the actual component template used in the BPGC. The template can be overridden in a child. */
 	ENGINE_API UActorComponent* GetActualComponentTemplate(class UBlueprintGeneratedClass* ActualBPGC) const;
diff --git a/Engine/Source/Runtime/Engine/Classes/Engine/SimpleConstructionScript.h b/Engine/Source/Runtime/Engine/Classes/Engine/SimpleConstructionScript.h
--- a/Engine/Source/Runtime/Engine/Classes/Engine/SimpleConstructionScript.h
+++ b/Engine/Source/Runtime/Engine/Classes/Engine/SimpleConstructionScript.h
@@ -42,7 +42,7 @@ class USimpleConstructionScript : public UObject
 	 * @param RootRelativeRotationCache Optional rotation cache to use when applying the transform to the root component.
 	 * @param bIsDefaultTransform	Indicates whether or not the given transform is a "default" transform, in which case it can be overridden by template defaults.
 	 */
-	void ExecuteScriptOnActor(AActor* Actor, const TInlineComponentArray<USceneComponent*>& NativeSceneComponents, const FTransform& RootTransform, const FRotationConversionCache* RootRelativeRotationCache, bool bIsDefaultTransform);
+	void ExecuteScriptOnActor(AActor* Actor, const TInlineComponentArray<USceneComponent*>& NativeSceneComponents, const FTransform& RootTransform, const FRotationConversionCache* RootRelativeRotationCache, bool bIsDefaultTransform, ESpawnActorScaleMethod TransformScaleMethod = ESpawnActorScaleMethod::OverrideRootScale);
 
 	/** Create the map from names to SCS_Nodes to improve FindSCSNode performance during construction script execution */
 	void CreateNameToSCSNodeMap();
diff --git a/Engine/Source/Runtime/Engine/Classes/Engine/World.h b/Engine/Source/Runtime/Engine/Classes/Engine/World.h
--- a/Engine/Source/Runtime/Engine/Classes/Engine/World.h
+++ b/Engine/Source/Runtime/Engine/Classes/Engine/World.h
@@ -518,6 +518,9 @@ struct ENGINE_API FActorSpawnParameters
 	/** Method for resolving collisions at the spawn point. Undefined means no override, use the actor's setting. */
 	ESpawnActorCollisionHandlingMethod SpawnCollisionHandlingOverride;
 
+	/** Determines whether to multiply or override root component with provided spawn transform */
+	ESpawnActorScaleMethod TransformScaleMethod = ESpawnActorScaleMethod::MultiplyWithRoot;
+
 private:
 
 	friend class UPackageMapClient;
@@ -3493,7 +3496,8 @@ public:
 		FTransform const& Transform,
 		AActor* Owner = nullptr,
 		APawn* Instigator = nullptr,
-		ESpawnActorCollisionHandlingMethod CollisionHandlingOverride = ESpawnActorCollisionHandlingMethod::Undefined
+		ESpawnActorCollisionHandlingMethod CollisionHandlingOverride = ESpawnActorCollisionHandlingMethod::Undefined,
+		ESpawnActorScaleMethod TransformScaleMethod = ESpawnActorScaleMethod::MultiplyWithRoot
 		)
 	{
 		if( Owner )
@@ -3502,6 +3506,7 @@ public:
 		}
 		FActorSpawnParameters SpawnInfo;
 		SpawnInfo.SpawnCollisionHandlingOverride = CollisionHandlingOverride;
+		SpawnInfo.TransformScaleMethod = TransformScaleMethod;
 		SpawnInfo.Owner = Owner;
 		SpawnInfo.Instigator = Instigator;
 		SpawnInfo.bDeferConstruction = true;
diff --git a/Engine/Source/Runtime/Engine/Classes/GameFramework/Actor.h b/Engine/Source/Runtime/Engine/Classes/GameFramework/Actor.h
--- a/Engine/Source/Runtime/Engine/Classes/GameFramework/Actor.h
+++ b/Engine/Source/Runtime/Engine/Classes/GameFramework/Actor.h
@@ -71,6 +71,16 @@ enum class EActorUpdateOverlapsMethod : uint8
 	NeverUpdate
 };
 
+/** Determines how the transform being passed into actor spawning methods interact with the actor's default root component */
+UENUM(BlueprintType)
+enum class ESpawnActorScaleMethod : uint8
+{
+	/** Ignore the default scale in the actor's root component and hard-set it to the value of SpawnTransform Parameter */
+	OverrideRootScale						UMETA(DisplayName = "Override Root Component Scale"),
+	/** Multiply value of the SpawnTransform Parameter with the default scale in the actor's root component */
+	MultiplyWithRoot						UMETA(DisplayName = "Multiply Scale With Root Component Scale"),
+};
+
 #if WITH_EDITORONLY_DATA
 /** Enum defining how actor will be placed in the partition */
 UENUM()
@@ -2738,10 +2748,10 @@ public:
 	virtual bool IsRelevancyOwnerFor(const AActor* ReplicatedActor, const AActor* ActorOwner, const AActor* ConnectionActor) const;
 
 	/** Called after the actor is spawned in the world.  Responsible for setting up actor for play. */
-	void PostSpawnInitialize(FTransform const& SpawnTransform, AActor* InOwner, APawn* InInstigator, bool bRemoteOwned, bool bNoFail, bool bDeferConstruction);
+	void PostSpawnInitialize(FTransform const& SpawnTransform, AActor* InOwner, APawn* InInstigator, bool bRemoteOwned, bool bNoFail, bool bDeferConstruction, ESpawnActorScaleMethod TransformScaleMethod = ESpawnActorScaleMethod::MultiplyWithRoot);
 
 	/** Called to finish the spawning process, generally in the case of deferred spawning */
-	void FinishSpawning(const FTransform& Transform, bool bIsDefaultTransform = false, const FComponentInstanceDataCache* InstanceDataCache = nullptr);
+	void FinishSpawning(const FTransform& Transform, bool bIsDefaultTransform = false, const FComponentInstanceDataCache* InstanceDataCache = nullptr, ESpawnActorScaleMethod TransformScaleMethod = ESpawnActorScaleMethod::OverrideRootScale);
 
 	/** Called after the actor has run its construction. Responsible for finishing the actor spawn process. */
 	void PostActorConstruction();
@@ -3051,7 +3061,7 @@ public:
 	 *
 	 * @return Returns false if the hierarchy was not error free and we've put the Actor is disaster recovery mode
 	 */
-	bool ExecuteConstruction(const FTransform& Transform, const struct FRotationConversionCache* TransformRotationCache, const class FComponentInstanceDataCache* InstanceDataCache, bool bIsDefaultTransform = false);
+	bool ExecuteConstruction(const FTransform& Transform, const struct FRotationConversionCache* TransformRotationCache, const class FComponentInstanceDataCache* InstanceDataCache, bool bIsDefaultTransform = false, ESpawnActorScaleMethod TransformScaleMethod = ESpawnActorScaleMethod::OverrideRootScale);
 
 	/**
 	 * Called when an instance of this class is placed (in editor) or spawned.
diff --git a/Engine/Source/Runtime/Engine/Classes/Kismet/GameplayStatics.h b/Engine/Source/Runtime/Engine/Classes/Kismet/GameplayStatics.h
--- a/Engine/Source/Runtime/Engine/Classes/Kismet/GameplayStatics.h
+++ b/Engine/Source/Runtime/Engine/Classes/Kismet/GameplayStatics.h
@@ -62,11 +62,11 @@ class ENGINE_API UGameplayStatics : public UBlueprintFunctionLibrary
 
 	/** Spawns an instance of an actor class, but does not automatically run its construction script.  */
 	UFUNCTION(BlueprintCallable, Category = "Spawning", meta = (WorldContext = "WorldContextObject", UnsafeDuringActorConstruction = "true", BlueprintInternalUseOnly = "true"))
-	static class AActor* BeginDeferredActorSpawnFromClass(const UObject* WorldContextObject, TSubclassOf<AActor> ActorClass, const FTransform& SpawnTransform, ESpawnActorCollisionHandlingMethod CollisionHandlingOverride = ESpawnActorCollisionHandlingMethod::Undefined, AActor* Owner = nullptr);
+	static class AActor* BeginDeferredActorSpawnFromClass(const UObject* WorldContextObject, TSubclassOf<AActor> ActorClass, const FTransform& SpawnTransform, ESpawnActorCollisionHandlingMethod CollisionHandlingOverride = ESpawnActorCollisionHandlingMethod::Undefined, AActor* Owner = nullptr, ESpawnActorScaleMethod TransformScaleMethod = ESpawnActorScaleMethod::MultiplyWithRoot);
 
 	/** 'Finish' spawning an actor.  This will run the construction script. */
 	UFUNCTION(BlueprintCallable, Category="Spawning", meta=(UnsafeDuringActorConstruction = "true", BlueprintInternalUseOnly = "true"))
-	static class AActor* FinishSpawningActor(class AActor* Actor, const FTransform& SpawnTransform);
+	static class AActor* FinishSpawningActor(class AActor* Actor, const FTransform& SpawnTransform, ESpawnActorScaleMethod TransformScaleMethod = ESpawnActorScaleMethod::MultiplyWithRoot);
 
 	// --- Actor functions ------------------------------
 
diff --git a/Engine/Source/Runtime/Engine/Private/Actor.cpp b/Engine/Source/Runtime/Engine/Private/Actor.cpp
--- a/Engine/Source/Runtime/Engine/Private/Actor.cpp
+++ b/Engine/Source/Runtime/Engine/Private/Actor.cpp
@@ -3627,7 +3627,7 @@ static void ValidateDeferredTransformCache()
 	}
 }
 
-void AActor::PostSpawnInitialize(FTransform const& UserSpawnTransform, AActor* InOwner, APawn* InInstigator, bool bRemoteOwned, bool bNoFail, bool bDeferConstruction)
+void AActor::PostSpawnInitialize(FTransform const& UserSpawnTransform, AActor* InOwner, APawn* InInstigator, bool bRemoteOwned, bool bNoFail, bool bDeferConstruction, ESpawnActorScaleMethod TransformScaleMethod)
 {
 	// General flow here is like so
 	// - Actor sets up the basics.
@@ -3664,7 +3664,16 @@ void AActor::PostSpawnInitialize(FTransform const& UserSpawnTransform, AActor* I
 		// Respect any non-default transform value that the root component may have received from the archetype that's owned
 		// by the native CDO, so the final transform might not always necessarily equate to the passed-in UserSpawnTransform.
 		const FTransform RootTransform(SceneRootComponent->GetRelativeRotation(), SceneRootComponent->GetRelativeLocation(), SceneRootComponent->GetRelativeScale3D());
-		const FTransform FinalRootComponentTransform = RootTransform * UserSpawnTransform;
+		FTransform FinalRootComponentTransform = RootTransform;
+		switch(TransformScaleMethod)
+		{
+		case ESpawnActorScaleMethod::OverrideRootScale:
+			FinalRootComponentTransform = UserSpawnTransform;
+			break;
+		case ESpawnActorScaleMethod::MultiplyWithRoot:
+			FinalRootComponentTransform = RootTransform * UserSpawnTransform;
+			break;
+		}
 		SceneRootComponent->SetWorldTransform(FinalRootComponentTransform, false, nullptr, ETeleportType::ResetPhysics);
 	}
 
@@ -3715,7 +3724,7 @@ void AActor::PostSpawnInitialize(FTransform const& UserSpawnTransform, AActor* I
 
 #include UE_INLINE_GENERATED_CPP_BY_NAME(Actor)
 
-void AActor::FinishSpawning(const FTransform& UserTransform, bool bIsDefaultTransform, const FComponentInstanceDataCache* InstanceDataCache)
+void AActor::FinishSpawning(const FTransform& UserTransform, bool bIsDefaultTransform, const FComponentInstanceDataCache* InstanceDataCache, ESpawnActorScaleMethod TransformScaleMethod)
 {
 #if ENABLE_SPAWNACTORTIMER
 	FScopedSpawnActorTimer SpawnTimer(GetClass()->GetFName(), ESpawnActorTimingType::FinishSpawning);
@@ -3759,7 +3768,7 @@ void AActor::FinishSpawning(const FTransform& UserTransform, bool bIsDefaultTran
 
 		{
 			FEditorScriptExecutionGuard ScriptGuard;
-			ExecuteConstruction(FinalRootComponentTransform, nullptr, InstanceDataCache, bIsDefaultTransform);
+			ExecuteConstruction(FinalRootComponentTransform, nullptr, InstanceDataCache, bIsDefaultTransform, TransformScaleMethod);
 		}
 
 		{
diff --git a/Engine/Source/Runtime/Engine/Private/ActorConstruction.cpp b/Engine/Source/Runtime/Engine/Private/ActorConstruction.cpp
--- a/Engine/Source/Runtime/Engine/Private/ActorConstruction.cpp
+++ b/Engine/Source/Runtime/Engine/Private/ActorConstruction.cpp
@@ -723,7 +723,7 @@ namespace
 	TMap<const AActor*, TMap<const UObject*, int32>, TInlineSetAllocator<4>> UCSBlueprintComponentArchetypeCounts;
 }
 
-bool AActor::ExecuteConstruction(const FTransform& Transform, const FRotationConversionCache* TransformRotationCache, const FComponentInstanceDataCache* InstanceDataCache, bool bIsDefaultTransform)
+bool AActor::ExecuteConstruction(const FTransform& Transform, const FRotationConversionCache* TransformRotationCache, const FComponentInstanceDataCache* InstanceDataCache, bool bIsDefaultTransform, ESpawnActorScaleMethod TransformScaleMethod)
 {
 	check(IsValid(this));
 	check(!HasAnyFlags(RF_BeginDestroyed|RF_FinishDestroyed));
@@ -797,7 +797,7 @@ bool AActor::ExecuteConstruction(const FTransform& Transform, const FRotationCon
 				if (SCS)
 				{
 					SCS->CreateNameToSCSNodeMap();
-					SCS->ExecuteScriptOnActor(this, NativeSceneComponents, Transform, TransformRotationCache, bIsDefaultTransform);
+					SCS->ExecuteScriptOnActor(this, NativeSceneComponents, Transform, TransformRotationCache, bIsDefaultTransform, TransformScaleMethod);
 				}
 				// Now that the construction scripts have been run, we can create timelines and hook them up
 				UBlueprintGeneratedClass::CreateComponentsForActor(CurrentBPGClass, this);
diff --git a/Engine/Source/Runtime/Engine/Private/GameplayStatics.cpp b/Engine/Source/Runtime/Engine/Private/GameplayStatics.cpp
--- a/Engine/Source/Runtime/Engine/Private/GameplayStatics.cpp
+++ b/Engine/Source/Runtime/Engine/Private/GameplayStatics.cpp
@@ -813,7 +813,7 @@ class AActor* UGameplayStatics::BeginSpawningActorFromBlueprint(const UObject* W
 	return nullptr;
 }
 
-class AActor* UGameplayStatics::BeginDeferredActorSpawnFromClass(const UObject* WorldContextObject, TSubclassOf<AActor> ActorClass, const FTransform& SpawnTransform, ESpawnActorCollisionHandlingMethod CollisionHandlingOverride /*= ESpawnActorCollisionHandlingMethod::Undefined*/, AActor* Owner /*= nullptr*/)
+class AActor* UGameplayStatics::BeginDeferredActorSpawnFromClass(const UObject* WorldContextObject, TSubclassOf<AActor> ActorClass, const FTransform& SpawnTransform, ESpawnActorCollisionHandlingMethod CollisionHandlingOverride /*= ESpawnActorCollisionHandlingMethod::Undefined*/, AActor* Owner /*= nullptr*/, ESpawnActorScaleMethod TransformScaleMethod /*= ESpawnActorScaleMethod::MultiplyWithRoot*/)
 {
 	SCOPE_CYCLE_COUNTER(STAT_SpawnTime);
 	if (UClass* Class = *ActorClass)
@@ -833,7 +833,7 @@ class AActor* UGameplayStatics::BeginDeferredActorSpawnFromClass(const UObject*
 
 		if (UWorld* World = GEngine->GetWorldFromContextObject(WorldContextObject, EGetWorldErrorMode::LogAndReturnNull))
 		{
-			return World->SpawnActorDeferred<AActor>(Class, SpawnTransform, Owner, AutoInstigator, CollisionHandlingOverride);
+			return World->SpawnActorDeferred<AActor>(Class, SpawnTransform, Owner, AutoInstigator, CollisionHandlingOverride, TransformScaleMethod);
 		}
 		else
 		{
@@ -848,12 +848,12 @@ class AActor* UGameplayStatics::BeginDeferredActorSpawnFromClass(const UObject*
 	return nullptr;
 }
 
-AActor* UGameplayStatics::FinishSpawningActor(AActor* Actor, const FTransform& SpawnTransform)
+AActor* UGameplayStatics::FinishSpawningActor(AActor* Actor, const FTransform& SpawnTransform, ESpawnActorScaleMethod TransformScaleMethod)
 {
 	SCOPE_CYCLE_COUNTER(STAT_SpawnTime);
 	if (Actor)
 	{
-		Actor->FinishSpawning(SpawnTransform);
+		Actor->FinishSpawning(SpawnTransform, false, nullptr, TransformScaleMethod);
 	}
 
 	return Actor;
diff --git a/Engine/Source/Runtime/Engine/Private/LevelActor.cpp b/Engine/Source/Runtime/Engine/Private/LevelActor.cpp
--- a/Engine/Source/Runtime/Engine/Private/LevelActor.cpp
+++ b/Engine/Source/Runtime/Engine/Private/LevelActor.cpp
@@ -726,7 +726,7 @@ AActor* UWorld::SpawnActor( UClass* Class, FTransform const* UserTransformPtr, c
 	// Broadcast delegate before the actor and its contained components are initialized
 	OnActorPreSpawnInitialization.Broadcast(Actor);
 
-	Actor->PostSpawnInitialize(UserTransform, SpawnParameters.Owner, SpawnParameters.Instigator, SpawnParameters.IsRemoteOwned(), SpawnParameters.bNoFail, SpawnParameters.bDeferConstruction);
+	Actor->PostSpawnInitialize(UserTransform, SpawnParameters.Owner, SpawnParameters.Instigator, SpawnParameters.IsRemoteOwned(), SpawnParameters.bNoFail, SpawnParameters.bDeferConstruction, SpawnParameters.TransformScaleMethod);
 	
 	// If we are spawning an external actor, mark this package dirty
 	if (ExternalPackage)
diff --git a/Engine/Source/Runtime/Engine/Private/SCS_Node.cpp b/Engine/Source/Runtime/Engine/Private/SCS_Node.cpp
--- a/Engine/Source/Runtime/Engine/Private/SCS_Node.cpp
+++ b/Engine/Source/Runtime/Engine/Private/SCS_Node.cpp
@@ -82,7 +82,7 @@ const FBlueprintCookedComponentInstancingData* USCS_Node::GetActualComponentTemp
 	return OverridenComponentTemplateData ? OverridenComponentTemplateData : &CookedComponentInstancingData;
 }
 
-UActorComponent* USCS_Node::ExecuteNodeOnActor(AActor* Actor, USceneComponent* ParentComponent, const FTransform* RootTransform, const FRotationConversionCache* RootRelativeRotationCache, bool bIsDefaultTransform)
+UActorComponent* USCS_Node::ExecuteNodeOnActor(AActor* Actor, USceneComponent* ParentComponent, const FTransform* RootTransform, const FRotationConversionCache* RootRelativeRotationCache, bool bIsDefaultTransform, ESpawnActorScaleMethod TransformScaleMethod)
 {
 	check(Actor != nullptr);
 	check(IsValid(ParentComponent) || (RootTransform != nullptr)); // must specify either a parent component or a world transform
@@ -129,12 +129,21 @@ UActorComponent* USCS_Node::ExecuteNodeOnActor(AActor* Actor, USceneComponent* P
 			if (!IsValid(ParentComponent) || ParentComponent == NewSceneComp)
 			{
 				FTransform WorldTransform = *RootTransform;
+				switch(TransformScaleMethod)
+				{
+				case ESpawnActorScaleMethod::OverrideRootScale:
+					// Use the provided transform and ignore the root component
+					break;
+				case ESpawnActorScaleMethod::MultiplyWithRoot:
+					WorldTransform = NewSceneComp->GetRelativeTransform() * WorldTransform;
+					break;
+				}
+				
 				if(bIsDefaultTransform)
 				{
 					// Note: We use the scale vector from the component template when spawning (to match what happens with a native root). This
 					// does NOT occur when this component is instanced as part of dynamically spawning a Blueprint class in a cooked build (i.e.
-					// 'bIsDefaultTransform' will be 'false' in that situation). In order to maintain the same behavior between a nativized and
-					// non-nativized cooked build, if this ever changes, we would also need to update the code in AActor::PostSpawnInitialize().
+					// 'bIsDefaultTransform' will be 'false' in that situation).
 					WorldTransform.SetScale3D(NewSceneComp->GetRelativeScale3D());
 				}
 
diff --git a/Engine/Source/Runtime/Engine/Private/SimpleConstructionScript.cpp b/Engine/Source/Runtime/Engine/Private/SimpleConstructionScript.cpp
--- a/Engine/Source/Runtime/Engine/Private/SimpleConstructionScript.cpp
+++ b/Engine/Source/Runtime/Engine/Private/SimpleConstructionScript.cpp
@@ -636,7 +636,7 @@ void USimpleConstructionScript::RegisterInstancedComponent(UActorComponent* Inst
 	}
 }
 
-void USimpleConstructionScript::ExecuteScriptOnActor(AActor* Actor, const TInlineComponentArray<USceneComponent*>& NativeSceneComponents, const FTransform& RootTransform, const FRotationConversionCache* RootRelativeRotationCache, bool bIsDefaultTransform)
+void USimpleConstructionScript::ExecuteScriptOnActor(AActor* Actor, const TInlineComponentArray<USceneComponent*>& NativeSceneComponents, const FTransform& RootTransform, const FRotationConversionCache* RootRelativeRotationCache, bool bIsDefaultTransform, ESpawnActorScaleMethod TransformScaleMethod)
 {
 	if(RootNodes.Num() > 0)
 	{
@@ -684,7 +684,7 @@ void USimpleConstructionScript::ExecuteScriptOnActor(AActor* Actor, const TInlin
 				}
 
 				// Create the new component instance and any child components it may have
-				RootNode->ExecuteNodeOnActor(Actor, ParentComponent != nullptr ? ParentComponent : RootComponent, &RootTransform, RootRelativeRotationCache, bIsDefaultTransform);
+				RootNode->ExecuteNodeOnActor(Actor, ParentComponent != nullptr ? ParentComponent : RootComponent, &RootTransform, RootRelativeRotationCache, bIsDefaultTransform, TransformScaleMethod);
 			}
 		}
 	}
