Update promotable operator nodes to only have one conversion menu instead of having a second submenu that lists all the types.

#rnx

diff --git a/Engine/Source/Editor/BlueprintGraph/Classes/K2Node_PromotableOperator.h b/Engine/Source/Editor/BlueprintGraph/Classes/K2Node_PromotableOperator.h
--- a/Engine/Source/Editor/BlueprintGraph/Classes/K2Node_PromotableOperator.h
+++ b/Engine/Source/Editor/BlueprintGraph/Classes/K2Node_PromotableOperator.h
@@ -161,7 +161,7 @@ private:
 	* @param Menu			The parent context menu
 	* @param ContextPin		The pin to convert
 	*/
-	void CreateConversionSubMenu(UToolMenu* Menu, UEdGraphPin* ContextPin) const;
+	void CreateConversionMenu(struct FToolMenuSection& ConversionSection, UEdGraphPin* ContextPin) const;
 
 	/** The name that this operation uses ("Add", "Multiply", etc) */
 	FName OperationName;
diff --git a/Engine/Source/Editor/BlueprintGraph/Private/K2Node_PromotableOperator.cpp b/Engine/Source/Editor/BlueprintGraph/Private/K2Node_PromotableOperator.cpp
--- a/Engine/Source/Editor/BlueprintGraph/Private/K2Node_PromotableOperator.cpp
+++ b/Engine/Source/Editor/BlueprintGraph/Private/K2Node_PromotableOperator.cpp
@@ -170,19 +170,20 @@ void UK2Node_PromotableOperator::GetNodeContextMenuActions(UToolMenu* Menu, UGra
 	// Add the pin conversion sub menu
 	if (CanConvertPinType(Context->Pin))
 	{
+		static const FName ConvNodeName = FName("PromotableOperatorPinConvs");
+		static const FText ConvNodeStr = LOCTEXT("PromotableOperatorPinConvs", "Pin Conversions");
+
+		FToolMenuSection& ConversionSection = Menu->AddSection(ConvNodeName, ConvNodeStr);
+
 		// Give the user an option to reset this node to wildcard
 		if (!FWildcardNodeUtils::IsWildcardPin(Context->Pin))
 		{
-			static const FName ResetWildcardName = FName("PromotableOperatorResetWildcardPinConvs");
-			static const FText ResetWildcardStr = LOCTEXT("PromotableOperatorResetWildcardPinConvs", "Reset To Wildcard");
-			FToolMenuSection& Section = Menu->AddSection(ResetWildcardName, ResetWildcardStr);
-
-			const FText ResetName = LOCTEXT("ResetFunction_Tooltip", "Reset this node to wildcard");
+			const FText ResetName = LOCTEXT("ResetFunction", "To Wildcard");
 
-			Section.AddMenuEntry(
+			ConversionSection.AddMenuEntry(
 				FName(ResetName.ToString()),
 				ResetName,
-				LOCTEXT("ResetToWildcardTooltip", "Reset this node to a wildcard state."),
+				LOCTEXT("ResetToWildcardTooltip", "Break all connections and reset this node to a wildcard state."),
 				FSlateIcon(),
 				FUIAction(
 					FExecuteAction::CreateUObject(const_cast<UK2Node_PromotableOperator*>(this), &UK2Node_PromotableOperator::ConvertPinType, const_cast<UEdGraphPin*>(Context->Pin), FWildcardNodeUtils::GetDefaultWildcardPinType())
@@ -190,22 +191,13 @@ void UK2Node_PromotableOperator::GetNodeContextMenuActions(UToolMenu* Menu, UGra
 			);
 		}
 
-		FToolMenuSection& Section = Menu->AddSection("K2NodePromoOpConversionGraphNode");
-		Section.AddSubMenu(
-			"ConvertPin",
-			LOCTEXT("ConvertPin", "Convert Pin..."),
-			LOCTEXT("ConvertPinTooltip", "Convert the selected pin to another type"),
-			FNewToolMenuDelegate::CreateUObject(this, &UK2Node_PromotableOperator::CreateConversionSubMenu, (UEdGraphPin*)Context->Pin)
-		);
+		CreateConversionMenu(ConversionSection, (UEdGraphPin*)Context->Pin);
 	}
 }
 
-void UK2Node_PromotableOperator::CreateConversionSubMenu(UToolMenu* Menu, UEdGraphPin* ContextPin) const
+void UK2Node_PromotableOperator::CreateConversionMenu(FToolMenuSection& ConversionSection, UEdGraphPin* PinToConvert) const
 {
-	check(ContextPin);
-
-	static const FName ConvNodeName = FName("PromotableOperatorPinConvs");
-	static const FText ConvNodeStr = LOCTEXT("PromotableOperatorPinConvs", "Pin Conversions");
+	check(PinToConvert);
 
 	// Gather what pin types could possibly be used for a conversion with this operator
 	TArray<UFunction*> AvailableFunctions;
@@ -214,13 +206,13 @@ void UK2Node_PromotableOperator::CreateConversionSubMenu(UToolMenu* Menu, UEdGra
 	const UEdGraphSchema_K2* Schema = GetDefault<UEdGraphSchema_K2>();
 
 	// If this is a split pin, then we need to convert the parent pin, not the child.
-	if (ContextPin->ParentPin != nullptr)
+	if (PinToConvert->ParentPin != nullptr)
 	{
-		ContextPin = ContextPin->ParentPin;
+		PinToConvert = PinToConvert->ParentPin;
 	}
 
 	// If we have a pin that matches our current type, then we can use it to see if we can still get a valid function
-	FEdGraphPinType OriginalContextType = ContextPin->PinType;
+	FEdGraphPinType OriginalContextType = PinToConvert->PinType;
 
 	for (const UFunction* Func : AvailableFunctions)
 	{
@@ -231,7 +223,7 @@ void UK2Node_PromotableOperator::CreateConversionSubMenu(UToolMenu* Menu, UEdGra
 
 			if (Schema->ConvertPropertyToPinType(Param, /* out */ ParamType))
 			{
-				if (FWildcardNodeUtils::IsWildcardPin(ContextPin) || FTypePromotion::IsValidPromotion(ParamType, ContextPin->PinType) || FTypePromotion::IsValidPromotion(ContextPin->PinType, ParamType))
+				if (FWildcardNodeUtils::IsWildcardPin(PinToConvert) || FTypePromotion::IsValidPromotion(ParamType, PinToConvert->PinType) || FTypePromotion::IsValidPromotion(PinToConvert->PinType, ParamType))
 				{
 					PossiblePromos.AddUnique(ParamType);
 				}
@@ -248,20 +240,19 @@ void UK2Node_PromotableOperator::CreateConversionSubMenu(UToolMenu* Menu, UEdGra
 	// Add the options to the context menu
 	for (const FEdGraphPinType& PinType : PossiblePromos)
 	{
-		FToolMenuSection& Section = Menu->AddSection(ConvNodeName, ConvNodeStr);
-
 		FFormatNamedArguments Args;
-		Args.Add(TEXT("PinType"), Schema->TypeToText(PinType));
+		Args.Add(TEXT("NewPinType"), Schema->TypeToText(PinType));
+		Args.Add(TEXT("CurrentPinType"), Schema->TypeToText(OriginalContextType));
 
-		const FText PinConversionName = FText::Format(LOCTEXT("CallFunction_Tooltip", "{PinType}"), Args);
+		const FText PinConversionName = FText::Format(LOCTEXT("CallFunction_Tooltip", "To {NewPinType}"), Args);
 
-		Section.AddMenuEntry(
+		ConversionSection.AddMenuEntry(
 			FName(PinConversionName.ToString()),
 			PinConversionName,
-			LOCTEXT("ConvertPinTypeTooltip", "Convert this pin type to the selected type"),
+			FText::Format(LOCTEXT("ConvertPinTypeTooltip", "Convert this pin type from '{CurrentPinType}' to '{NewPinType}'"), Args),
 			FSlateIcon(),
 			FUIAction(
-				FExecuteAction::CreateUObject(const_cast<UK2Node_PromotableOperator*>(this), &UK2Node_PromotableOperator::ConvertPinType, const_cast<UEdGraphPin*>(ContextPin), PinType)
+				FExecuteAction::CreateUObject(const_cast<UK2Node_PromotableOperator*>(this), &UK2Node_PromotableOperator::ConvertPinType, const_cast<UEdGraphPin*>(PinToConvert), PinType)
 			)
 		);
 	}
@@ -279,14 +270,10 @@ bool UK2Node_PromotableOperator::CanConvertPinType(const UEdGraphPin* Pin) const
 
 FText UK2Node_PromotableOperator::GetTooltipText() const
 {
-	// If there are no connections then just display the op name
-	if (!HasAnyConnectionsOrDefaults())
-	{
-		return FTypePromotion::GetUserFacingOperatorName(OperationName);
-	}
+	FFormatNamedArguments Args;
+	Args.Add(TEXT("FunctionTooltip"), !HasAnyConnectionsOrDefaults() ? FTypePromotion::GetUserFacingOperatorName(OperationName) : Super::GetTooltipText());
 
-	// Otherwise use the default one (a more specific function tooltip)
-	return Super::GetTooltipText();
+	return FText::Format(LOCTEXT("PromotableOpTooltipText", "{FunctionTooltip}\n\nTo go back to old math nodes, uncheck 'Enable Type Promotion' in the Blueprint Editor Settings. Both node types are supported."), Args);
 }
 
 void UK2Node_PromotableOperator::PinDefaultValueChanged(UEdGraphPin* Pin)
