[BugFix] Crash diffing a blueprint with a Spawn Actor node

diff --git a/Engine/Source/Editor/BlueprintGraph/Private/K2Node_SpawnActorFromClass.cpp b/Engine/Source/Editor/BlueprintGraph/Private/K2Node_SpawnActorFromClass.cpp
--- a/Engine/Source/Editor/BlueprintGraph/Private/K2Node_SpawnActorFromClass.cpp
+++ b/Engine/Source/Editor/BlueprintGraph/Private/K2Node_SpawnActorFromClass.cpp
@@ -85,7 +85,7 @@ void UK2Node_SpawnActorFromClass::AllocateDefaultPins()
 	// Pin to set transform scaling behavior (ie whether to multiply scale with root component or to just ignore the root component default scale)
 	UEnum* const ScaleMethodEnum = StaticEnum<ESpawnActorScaleMethod>();
 	UEdGraphPin* const ScaleMethodPin = CreatePin(EGPD_Input, UEdGraphSchema_K2::PC_Byte, ScaleMethodEnum, FK2Node_SpawnActorFromClassHelper::TransformScaleMethodPinName);
-	ScaleMethodPin->DefaultValue = ScaleMethodEnum->GetNameStringByValue(static_cast<int>(ESpawnActorScaleMethod::MultiplyWithRoot));
+	ScaleMethodPin->DefaultValue = ScaleMethodEnum->GetNameStringByValue(static_cast<int>(ESpawnActorScaleMethod::SelectDefaultAtRuntime));
 	ScaleMethodPin->bAdvancedView = true;
 	
 	UEdGraphPin* OwnerPin = CreatePin(EGPD_Input, UEdGraphSchema_K2::PC_Object, AActor::StaticClass(), FK2Node_SpawnActorFromClassHelper::OwnerPinName);
@@ -367,6 +367,13 @@ void UK2Node_SpawnActorFromClass::PostLoad()
 
 void UK2Node_SpawnActorFromClass::FixupScaleMethodPin()
 {
+	// if this node is being diffed, don't fix up anything. Keep the legacy pins
+	const UPackage* Package = GetPackage();
+	if (Package && Package->HasAnyPackageFlags(PKG_ForDiffing))
+	{
+		return;
+	}
+
 	if (GetLinkerCustomVersion(FUE5MainStreamObjectVersion::GUID) < FUE5MainStreamObjectVersion::SpawnActorFromClassTransformScaleMethod)
 	{
 		const UEdGraphPin* const ClassPin = FindPin(TEXT( "Class" ));
