Avoid doing a full Blueprint action database refresh on the next UI tick after incurring one or more dynamic module loads in the editor.

diff --git a/Engine/Source/Editor/BlueprintGraph/Private/BlueprintActionDatabase.cpp b/Engine/Source/Editor/BlueprintGraph/Private/BlueprintActionDatabase.cpp
--- a/Engine/Source/Editor/BlueprintGraph/Private/BlueprintActionDatabase.cpp
+++ b/Engine/Source/Editor/BlueprintGraph/Private/BlueprintActionDatabase.cpp
@@ -564,6 +564,19 @@ namespace BlueprintActionDatabaseImpl
 	 */
 	TSet<TWeakObjectPtr<UObject>> PendingDelete;
 
+	/**
+	 * Modules that were explicitly loaded at runtime but have not yet been
+	 * registered into the database. These will be processed on the next tick.
+	 */
+	TSet<FName> PendingModules;
+
+	/**
+	 * Modules that were explicitly loaded at runtime and registered into
+	 * the database. We keep track of them here so that in the off chance
+	 * we allow it to be unloaded, we can then trigger a database refresh.
+	 */
+	TSet<FName> LoadedModules;
+
 	/** */
 	bool bIsInitializing = false;
 
@@ -574,9 +587,33 @@ namespace BlueprintActionDatabaseImpl
 //------------------------------------------------------------------------------
 static void BlueprintActionDatabaseImpl::OnModulesChanged(FName InModuleName, EModuleChangeReason InModuleChangeReason)
 {
-	if (InModuleChangeReason == EModuleChangeReason::ModuleLoaded || InModuleChangeReason == EModuleChangeReason::ModuleUnloaded)
+	switch (InModuleChangeReason)
 	{
-		BlueprintActionDatabaseImpl::bRefreshAllRequested = true;
+	case EModuleChangeReason::ModuleLoaded:
+		// If not already tracked, add it to the list of modules that need to be registered on the next tick.
+		if (!LoadedModules.Contains(InModuleName))
+		{
+			PendingModules.Add(InModuleName);
+		}
+		break;
+
+	case EModuleChangeReason::ModuleUnloaded:
+		// If pending, it was unloaded in the same frame, so we just need to remove it, and no refresh is needed.
+		if (!PendingModules.Remove(InModuleName))
+		{
+			// If already registered as a loaded module, then we need to remove it and do a full refresh on the next tick.
+			if (LoadedModules.Remove(InModuleName))
+			{
+				bRefreshAllRequested = true;
+			}
+		}
+
+		// Guard against the possibility of unloading a pending module that is also already registered.
+		checkf(!LoadedModules.Contains(InModuleName), TEXT("Module %s was unloaded, but wasn't unregistered from the Blueprint action database."), *InModuleName.ToString());
+		break;
+	
+	default:
+		break;
 	}
 }
 
@@ -1328,6 +1365,45 @@ void FBlueprintActionDatabase::Tick(float DeltaTime)
 	{
 		RefreshAll();
 	}
+	else if(!BlueprintActionDatabaseImpl::PendingModules.IsEmpty())
+	{
+		TRACE_CPUPROFILER_EVENT_SCOPE(FBlueprintActionDatabase::ProcessLoadedModules);
+
+		// Register actions for any new modules that were explicitly loaded prior to this tick.
+		for (const FName& LoadedModule : BlueprintActionDatabaseImpl::PendingModules)
+		{
+			if (BlueprintActionDatabaseImpl::LoadedModules.Contains(LoadedModule))
+			{
+				continue;
+			}
+
+			// Look for a script package that's associated with this module load. If there isn't one, we can skip it.
+			const FName ModuleScriptPackageName = FPackageName::GetModuleScriptPackageName(LoadedModule);
+			if (const UPackage* ModuleScriptPackage = FindPackage(nullptr, *ModuleScriptPackageName.ToString()))
+			{
+				TArray<UObject*> ObjectsToProcess;
+				const bool bIncludeNestedObjects = false;
+				GetObjectsWithPackage(ModuleScriptPackage, ObjectsToProcess, bIncludeNestedObjects, RF_ClassDefaultObject);
+				for (UObject* Object : ObjectsToProcess)
+				{
+					if (UClass* ObjectAsClass = Cast<UClass>(Object))
+					{
+						RefreshClassActions(ObjectAsClass);
+					}
+					else if (Object->IsA<UScriptStruct>() || Object->IsA<UEnum>())
+					{
+						RefreshAssetActions(Object);
+					}
+				}
+
+				// Only need to track modules that contain a script package.
+				BlueprintActionDatabaseImpl::LoadedModules.Add(LoadedModule);
+			}
+		}
+	}
+
+	// All pending modules should be fully registered at this point.
+	BlueprintActionDatabaseImpl::PendingModules.Empty();
 
 	// entries that were removed from the database, in preparation for a delete
 	// (but the user ended up not deleting the object)
