[BugFix] Clicking "Find References" on the Custom Event sidebar will search for "None"

diff --git a/Engine/Source/Editor/BlueprintGraph/Private/K2Node_Event.cpp b/Engine/Source/Editor/BlueprintGraph/Private/K2Node_Event.cpp
--- a/Engine/Source/Editor/BlueprintGraph/Private/K2Node_Event.cpp
+++ b/Engine/Source/Editor/BlueprintGraph/Private/K2Node_Event.cpp
@@ -884,14 +884,22 @@ FSlateIcon UK2Node_Event::GetIconAndTint(FLinearColor& OutColor) const
 
 FString UK2Node_Event::GetFindReferenceSearchString() const
 {
-	FString FunctionName = EventReference.GetMemberName().ToString(); // If we fail to find the function, still want to search for its expected name.
+	// Behavior modeled after UK2Node_Event::GetNodeTitle but without "Event" prepended
+	if (bOverrideFunction || (CustomFunctionName == NAME_None))
+	{
+		FString FunctionName = EventReference.GetMemberName().ToString(); // If we fail to find the function, still want to search for its expected name.
 
-	if (UFunction* Function = EventReference.ResolveMember<UFunction>(GetBlueprintClassFromNode()))
+		if (const UFunction* Function = EventReference.ResolveMember<UFunction>(GetBlueprintClassFromNode()))
+		{
+			FunctionName = UEdGraphSchema_K2::GetFriendlySignatureName(Function).ToString();
+		}
+
+		return FunctionName;
+	}
+	else
 	{
-		FunctionName = UEdGraphSchema_K2::GetFriendlySignatureName(Function).ToString();
+		return CustomFunctionName.ToString();
 	}
-
-	return FunctionName;
 }
 
 void UK2Node_Event::FindDiffs(UEdGraphNode* OtherNode, struct FDiffResults& Results)
