UE-169957 - Rewiring split pins after a type change creates an orphaned pin

Allow rewiring split pins that have had their type changed, which ERedirectType_Name permits. Oddly, the Break Struct nodes already work in this case, but that's because they completely bypass UK2Node_Variable::DoPinsMatchForReconstruction, and call the implementation in UK2Node::DoPinsMatchForReconstruction instead.

diff --git a/Engine/Source/Editor/BlueprintGraph/Private/K2Node_Variable.cpp b/Engine/Source/Editor/BlueprintGraph/Private/K2Node_Variable.cpp
--- a/Engine/Source/Editor/BlueprintGraph/Private/K2Node_Variable.cpp
+++ b/Engine/Source/Editor/BlueprintGraph/Private/K2Node_Variable.cpp
@@ -300,49 +300,49 @@ UK2Node::ERedirectType UK2Node_Variable::DoPinsMatchForReconstruction( const UEd
 		const FEdGraphPinType& InputType  = (OldPin->Direction == EGPD_Output) ? OldPin->PinType : NewPin->PinType;
 		const FEdGraphPinType& OutputType = (OldPin->Direction == EGPD_Output) ? NewPin->PinType : OldPin->PinType;
 
-		if (K2Schema->ArePinTypesCompatible(OutputType, InputType))
+		if (NewPin->ParentPin)
 		{
-			// If these are split pins, we need to do some name checking logic
-			if (NewPin->ParentPin)
+			// If the OldPin is not split, then these don't match
+			if (OldPin->ParentPin == nullptr)
 			{
-				// If the OldPin is not split, then these don't match
-				if (OldPin->ParentPin == nullptr)
-				{
-					return ERedirectType_None;
-				}
+				return ERedirectType_None;
+			}
 
-				// Go through and find the original variable pin.
-				// If the number of steps out to the original variable pin is not the same then these don't match
-				const UEdGraphPin* ParentmostNewPin = NewPin;
-				const UEdGraphPin* ParentmostOldPin = OldPin;
+			// Go through and find the original variable pin.
+			// If the number of steps out to the original variable pin is not the same then these don't match
+			const UEdGraphPin* ParentmostNewPin = NewPin;
+			const UEdGraphPin* ParentmostOldPin = OldPin;
 
-				while (ParentmostNewPin->ParentPin)
-				{
-					if (ParentmostOldPin->ParentPin == nullptr)
-					{
-						return ERedirectType_None;
-					}
-					ParentmostNewPin = ParentmostNewPin->ParentPin;
-					ParentmostOldPin = ParentmostOldPin->ParentPin;
-				}
-
-				if (ParentmostOldPin->ParentPin)
+			while (ParentmostNewPin->ParentPin)
+			{
+				if (ParentmostOldPin->ParentPin == nullptr)
 				{
 					return ERedirectType_None;
 				}
+				ParentmostNewPin = ParentmostNewPin->ParentPin;
+				ParentmostOldPin = ParentmostOldPin->ParentPin;
+			}
 
-				// Compare whether the names, ignoring the original variable's name in the case of renames, match
-				FName NewPinPropertyName = FName(*NewPin->PinName.ToString().RightChop(ParentmostNewPin->PinName.ToString().Len() + 1));
-				FName OldPinPropertyName = FName(*OldPin->PinName.ToString().RightChop(ParentmostOldPin->PinName.ToString().Len() + 1));
+			if (ParentmostOldPin->ParentPin)
+			{
+				return ERedirectType_None;
+			}
 
-				if (!DoesRenamedVariableMatch(OldPinPropertyName, NewPinPropertyName, Cast<UStruct>(NewPin->ParentPin->PinType.PinSubCategoryObject.Get())))
-				{
-					return ERedirectType_None;
-				}
+			// Compare whether the names, ignoring the original variable's name in the case of renames, match
+			FName NewPinPropertyName = FName(*NewPin->PinName.ToString().RightChop(ParentmostNewPin->PinName.ToString().Len() + 1));
+			FName OldPinPropertyName = FName(*OldPin->PinName.ToString().RightChop(ParentmostOldPin->PinName.ToString().Len() + 1));
+
+			if (!DoesRenamedVariableMatch(OldPinPropertyName, NewPinPropertyName, Cast<UStruct>(NewPin->ParentPin->PinType.PinSubCategoryObject.Get())))
+			{
+				return ERedirectType_None;
 			}
 
 			return ERedirectType_Name;
 		}
+		else if (K2Schema->ArePinTypesCompatible(OutputType, InputType))
+		{
+			return ERedirectType_Name;
+		}
 		else
 		{
 			const bool bNewPinIsObject = (NewPin->PinType.PinCategory == UEdGraphSchema_K2::PC_Object);
