Improve tooltip on function call nodes when a function cannot be resolved

diff --git a/Engine/Source/Editor/BlueprintGraph/Private/K2Node_CallFunction.cpp b/Engine/Source/Editor/BlueprintGraph/Private/K2Node_CallFunction.cpp
--- a/Engine/Source/Editor/BlueprintGraph/Private/K2Node_CallFunction.cpp
+++ b/Engine/Source/Editor/BlueprintGraph/Private/K2Node_CallFunction.cpp
@@ -1493,7 +1493,15 @@ FText UK2Node_CallFunction::GetTooltipText() const
 	UFunction* Function = GetTargetFunction();
 	if (Function == nullptr)
 	{
-		return FText::Format(LOCTEXT("CallUnknownFunction", "Call unknown function {0}"), FText::FromName(FunctionReference.GetMemberName()));
+		// try to see where this function is meant to come from:
+		if (UClass* FuncOwnerClass = FunctionReference.GetMemberParentClass(GetBlueprintClassFromNode()))
+		{
+			return FText::Format(LOCTEXT("CallUnknownFunction", "Call unknown function {0} - missing from {1}"), FText::FromName(FunctionReference.GetMemberName()), FText::FromName(FuncOwnerClass->GetFName()));
+		}
+		else
+		{
+			return FText::Format(LOCTEXT("CallUnknownFunction", "Call unknown function {0}"), FText::FromName(FunctionReference.GetMemberName()));
+		}
 	}
 	else if (CachedTooltip.IsOutOfDate(this))
 	{
