Fix a regression that led to missing actions in the Blueprint graph context menu (introduced at 23029169).

Also restore previous Blueprint editor-related workflow perf improvements that were temporarily backed out to work around the regression (23125190, 23177402).

#rnx

diff --git a/Engine/Source/Editor/BlueprintGraph/Private/BlueprintActionDatabase.cpp b/Engine/Source/Editor/BlueprintGraph/Private/BlueprintActionDatabase.cpp
--- a/Engine/Source/Editor/BlueprintGraph/Private/BlueprintActionDatabase.cpp
+++ b/Engine/Source/Editor/BlueprintGraph/Private/BlueprintActionDatabase.cpp
@@ -564,6 +564,19 @@ namespace BlueprintActionDatabaseImpl
 	 */
 	TSet<TWeakObjectPtr<UObject>> PendingDelete;
 
+	/**
+	 * Modules that were explicitly loaded at runtime but have not yet been
+	 * registered into the database. These will be processed on the next tick.
+	 */
+	TSet<FName> PendingModules;
+
+	/**
+	 * Modules that were explicitly loaded at runtime and registered into
+	 * the database. We keep track of them here so that in the off chance
+	 * we allow it to be unloaded, we can then trigger a database refresh.
+	 */
+	TSet<FName> LoadedModules;
+
 	/** */
 	bool bIsInitializing = false;
 
@@ -574,9 +587,33 @@ namespace BlueprintActionDatabaseImpl
 //------------------------------------------------------------------------------
 static void BlueprintActionDatabaseImpl::OnModulesChanged(FName InModuleName, EModuleChangeReason InModuleChangeReason)
 {
-	if (InModuleChangeReason == EModuleChangeReason::ModuleLoaded || InModuleChangeReason == EModuleChangeReason::ModuleUnloaded)
+	switch (InModuleChangeReason)
 	{
-		BlueprintActionDatabaseImpl::bRefreshAllRequested = true;
+	case EModuleChangeReason::ModuleLoaded:
+		// If not already tracked, add it to the list of modules that need to be registered on the next tick.
+		if (!LoadedModules.Contains(InModuleName))
+		{
+			PendingModules.Add(InModuleName);
+		}
+		break;
+
+	case EModuleChangeReason::ModuleUnloaded:
+		// If pending, it was unloaded in the same frame, so we just need to remove it, and no refresh is needed.
+		if (!PendingModules.Remove(InModuleName))
+		{
+			// If already registered as a loaded module, then we need to remove it and do a full refresh on the next tick.
+			if (LoadedModules.Remove(InModuleName))
+			{
+				bRefreshAllRequested = true;
+			}
+		}
+
+		// Guard against the possibility of unloading a pending module that is also already registered.
+		checkf(!LoadedModules.Contains(InModuleName), TEXT("Module %s was unloaded, but wasn't unregistered from the Blueprint action database."), *InModuleName.ToString());
+		break;
+	
+	default:
+		break;
 	}
 }
 
@@ -1328,6 +1365,63 @@ void FBlueprintActionDatabase::Tick(float DeltaTime)
 	{
 		RefreshAll();
 	}
+	
+	// Check for any modules that may have been loaded since the last tick. Even if we call RefreshAll() above, we still want to run
+	// through this list in order to keep track of loaded modules containing native script types that are registered into the database.
+	if(!BlueprintActionDatabaseImpl::PendingModules.IsEmpty())
+	{
+		TRACE_CPUPROFILER_EVENT_SCOPE(FBlueprintActionDatabase::ProcessLoadedModules);
+
+		// Treat this as an initialization event so that the registrar is configured to register new types rather than to refresh existing ones. For
+		// example, if this module contains a new node type that implements a GetMenuActions() override, most implementations assume that a NULL action
+		// key filter indicates the initialization path. This ensures we get the same behavior as node types that are registered via the RefreshAll() API.
+		TGuardValue<bool> ScopedInitialization(BlueprintActionDatabaseImpl::bIsInitializing, true);
+
+		// Register actions for any new modules that were explicitly loaded prior to this tick. Note that native type objects defined within the module
+		// may have already been registered prior to receiving the load event; in that case, action(s) associated with those objects are already present.
+		for (const FName& LoadedModule : BlueprintActionDatabaseImpl::PendingModules)
+		{
+			if (BlueprintActionDatabaseImpl::LoadedModules.Contains(LoadedModule))
+			{
+				continue;
+			}
+
+			// Look for a script package that's associated with this module load. If there isn't one, we can skip it.
+			const FName ModuleScriptPackageName = FPackageName::GetModuleScriptPackageName(LoadedModule);
+			if (const UPackage* ModuleScriptPackage = FindPackage(nullptr, *ModuleScriptPackageName.ToString()))
+			{
+				TArray<UObject*> ObjectsToProcess;
+				const bool bIncludeNestedObjects = false;
+				GetObjectsWithPackage(ModuleScriptPackage, ObjectsToProcess, bIncludeNestedObjects, RF_ClassDefaultObject);
+				for (UObject* Object : ObjectsToProcess)
+				{
+					UClass* ObjectAsClass = Cast<UClass>(Object);
+					const bool bIsNativeTypeObject = ObjectAsClass != nullptr || Object->IsA<UScriptStruct>() || Object->IsA<UEnum>();
+					
+					// Only need to include native types and those not already added through the registrar at initialization time.
+					if (!bIsNativeTypeObject || ActionRegistry.Contains(Object))
+					{
+						continue;
+					}
+
+					if (ObjectAsClass)
+					{
+						RefreshClassActions(ObjectAsClass);
+					}
+					else
+					{
+						RefreshAssetActions(Object);
+					}
+				}
+
+				// We only need to track modules that contain a script package in the off-chance that it is ever unloaded, in which case we'd need
+				// to refresh the database to account for any types that go missing as a result. Otherwise, we can ignore this module when unloaded.
+				BlueprintActionDatabaseImpl::LoadedModules.Add(LoadedModule);
+			}
+		}
+
+		BlueprintActionDatabaseImpl::PendingModules.Empty();
+	}
 
 	// entries that were removed from the database, in preparation for a delete
 	// (but the user ended up not deleting the object)
diff --git a/Engine/Source/Editor/Kismet/Private/BlueprintEditor.cpp b/Engine/Source/Editor/Kismet/Private/BlueprintEditor.cpp
--- a/Engine/Source/Editor/Kismet/Private/BlueprintEditor.cpp
+++ b/Engine/Source/Editor/Kismet/Private/BlueprintEditor.cpp
@@ -1979,10 +1979,16 @@ void FBlueprintEditor::CommonInitialization(const TArray<UBlueprint*>& InitBluep
 
 	if (InitBlueprints.Num() == 1)
 	{
-		// Load blueprint libraries
-		if (!bShouldOpenInDefaultsMode && ShouldLoadBPLibrariesFromAssetRegistry())
+		if (!bShouldOpenInDefaultsMode)
 		{
-			LoadLibrariesFromAssetRegistry();
+			// Load blueprint libraries
+			if (ShouldLoadBPLibrariesFromAssetRegistry())
+			{
+				LoadLibrariesFromAssetRegistry();
+			}
+
+			// Init the action DB for the context menu/palette if not already constructed
+			FBlueprintActionDatabase::Get();
 		}
 
 		FLoadObjectsFromAssetRegistryHelper::Load<UUserDefinedEnum>(UserDefinedEnumerators);
@@ -2880,10 +2886,6 @@ void FBlueprintEditor::CreateDefaultTabContents(const TArray<UBlueprint*>& InBlu
 
 	if (InBlueprint)
 	{
-		this->Palette = 
-			SNew(SBlueprintPalette, SharedThis(this))
-				.IsEnabled(this, &FBlueprintEditor::IsFocusedGraphEditable);
-
 		this->BookmarksWidget =
 			SNew(SBlueprintBookmarks)
 				.EditorContext(SharedThis(this));
@@ -3677,6 +3679,19 @@ void FBlueprintEditor::SetKeyboardFocus()
 	FSlateApplication::Get().SetKeyboardFocus(GetMyBlueprintWidget());
 }
 
+TSharedRef<SBlueprintPalette> FBlueprintEditor::GetPalette()
+{
+	// Note: construction is deferred until first access; in large-scale projects this can be an expensive widget to construct during editor
+	// initialization logic. It's an unnecessary cost if the tab it's hosted in is closed and/or unavailable (e.g. as in the default layout).
+	if (!Palette.IsValid())
+	{
+		SAssignNew(Palette, SBlueprintPalette, SharedThis(this))
+			.IsEnabled(this, &FBlueprintEditor::IsFocusedGraphEditable);
+	}
+
+	return Palette.ToSharedRef();
+}
+
 bool FBlueprintEditor::TransactionObjectAffectsBlueprint(UObject* InTransactedObject)
 {
 	check(InTransactedObject);
diff --git a/Engine/Source/Editor/Kismet/Private/SBlueprintPalette.cpp b/Engine/Source/Editor/Kismet/Private/SBlueprintPalette.cpp
--- a/Engine/Source/Editor/Kismet/Private/SBlueprintPalette.cpp
+++ b/Engine/Source/Editor/Kismet/Private/SBlueprintPalette.cpp
@@ -9,7 +9,6 @@
 #include "AnimationStateMachineGraph.h"
 #include "AnimationStateMachineSchema.h"
 #include "AnimationTransitionGraph.h"
-#include "AssetDiscoveryIndicator.h"
 #include "AssetRegistry/AssetData.h"
 #include "AssetRegistry/AssetRegistryModule.h"
 #include "AssetRegistry/IAssetRegistry.h"
@@ -38,7 +37,6 @@
 #include "EdGraphNode_Comment.h"
 #include "EdGraphSchema_K2.h"
 #include "EdGraphSchema_K2_Actions.h"
-#include "EditorWidgetsModule.h"
 #include "Engine/Blueprint.h"
 #include "Engine/MemberReference.h"
 #include "Engine/SCS_Node.h"
@@ -67,6 +65,7 @@
 #include "Misc/FileHelper.h"
 #include "Misc/PackageName.h"
 #include "Misc/Paths.h"
+#include "Misc/ScopedSlowTask.h"
 #include "Modules/ModuleManager.h"
 #include "SBlueprintFavoritesPalette.h"
 #include "SBlueprintLibraryPalette.h"
@@ -2183,9 +2182,11 @@ TSharedPtr<SToolTip> SBlueprintPaletteItem::ConstructToolTipWidget() const
 //------------------------------------------------------------------------------
 void SBlueprintPalette::Construct(const FArguments& InArgs, TWeakPtr<FBlueprintEditor> InBlueprintEditor)
 {
-	// Create the asset discovery indicator
-	FEditorWidgetsModule& EditorWidgetsModule = FModuleManager::LoadModuleChecked<FEditorWidgetsModule>("EditorWidgets");
-	TSharedRef<SWidget> AssetDiscoveryIndicator = EditorWidgetsModule.CreateAssetDiscoveryIndicator(EAssetDiscoveryIndicatorScaleMode::Scale_Vertical);
+	const float NumProgressFrames = 2.0f;
+	const float SecondsToWaitBeforeShowingProgressDialog = 0.25f;
+
+	FScopedSlowTask SlowTask(NumProgressFrames, LOCTEXT("ConstructingPaletteTabContent", "Initializing Palette..."));
+	SlowTask.MakeDialogDelayed(SecondsToWaitBeforeShowingProgressDialog);
 
 	float FavoritesHeightRatio = 0.33f;
 	GConfig->GetFloat(*BlueprintPalette::ConfigSection, *BlueprintPalette::FavoritesHeightConfigKey, FavoritesHeightRatio, GEditorPerProjectIni);
@@ -2195,16 +2196,27 @@ void SBlueprintPalette::Construct(const FArguments& InArgs, TWeakPtr<FBlueprintE
 	bool bUseLegacyLayout = false;
 	GConfig->GetBool(*BlueprintPalette::ConfigSection, TEXT("bUseLegacyLayout"), bUseLegacyLayout, GEditorIni);
 
+	SlowTask.EnterProgressFrame();
+	TSharedRef<SWidget> FavoritesContent = SNew(SBlueprintFavoritesPalette, InBlueprintEditor);
+
+	SlowTask.EnterProgressFrame();
+	TSharedRef<SWidget> LibraryContent = SNew(SBlueprintLibraryPalette, InBlueprintEditor)
+		.UseLegacyLayout(bUseLegacyLayout);
+
 	if (bUseLegacyLayout)
 	{
+		LibraryWrapper = LibraryContent;
+
 		this->ChildSlot
 		[
-			SAssignNew(LibraryWrapper, SBlueprintLibraryPalette, InBlueprintEditor)
-				.UseLegacyLayout(bUseLegacyLayout)
+			LibraryContent
 		];
 	}
 	else 
 	{
+		LibraryContent->AddMetadata<FTagMetaData>(MakeShared<FTagMetaData>(TEXT("BlueprintPaletteLibrary")));
+		FavoritesContent->AddMetadata<FTagMetaData>(MakeShared<FTagMetaData>(TEXT("BlueprintPaletteFavorites")));
+
 		this->ChildSlot
 		[
 			SAssignNew(PaletteSplitter, SSplitter)
@@ -2215,15 +2227,13 @@ void SBlueprintPalette::Construct(const FArguments& InArgs, TWeakPtr<FBlueprintE
 			+ SSplitter::Slot()
 			.Value(FavoritesHeightRatio)
 			[
-				SNew(SBlueprintFavoritesPalette, InBlueprintEditor)
-				.AddMetaData<FTagMetaData>(FTagMetaData(TEXT("BlueprintPaletteFavorites")))
+				FavoritesContent
 			]
 
 			+ SSplitter::Slot()
 			.Value(LibraryHeightRatio)
 			[
-				SNew(SBlueprintLibraryPalette, InBlueprintEditor)
-				.AddMetaData<FTagMetaData>(FTagMetaData(TEXT("BlueprintPaletteLibrary")))
+				LibraryContent
 			]
 		];
 	}	
diff --git a/Engine/Source/Editor/Kismet/Public/BlueprintEditor.h b/Engine/Source/Editor/Kismet/Public/BlueprintEditor.h
--- a/Engine/Source/Editor/Kismet/Public/BlueprintEditor.h
+++ b/Engine/Source/Editor/Kismet/Public/BlueprintEditor.h
@@ -329,7 +329,7 @@ public:
 	/** Getters for the various Kismet2 widgets */
 	TSharedRef<SKismetInspector> GetInspector() const { return Inspector.ToSharedRef(); }
 	TSharedRef<SKismetInspector> GetDefaultEditor() const { return DefaultEditor.ToSharedRef(); }
-	TSharedRef<SBlueprintPalette> GetPalette() const { return Palette.ToSharedRef(); }
+	TSharedRef<SBlueprintPalette> GetPalette();
 	TSharedRef<SBlueprintBookmarks> GetBookmarksWidget() const { return BookmarksWidget.ToSharedRef(); }
 	TSharedRef<SWidget> GetCompilerResults() const { return CompilerResults.ToSharedRef(); }
 	TSharedRef<SFindInBlueprints> GetFindResults() const { return FindResults.ToSharedRef(); }
