Block non-array container pin types from being connected to 'self' pins on node types that support a foreach-style call site expansion.

This avoids a runtime crash that otherwise results due to a lack of compile-time expansion support for non-array container types.

diff --git a/Engine/Source/Editor/BlueprintGraph/Private/EdGraphSchema_K2.cpp b/Engine/Source/Editor/BlueprintGraph/Private/EdGraphSchema_K2.cpp
--- a/Engine/Source/Editor/BlueprintGraph/Private/EdGraphSchema_K2.cpp
+++ b/Engine/Source/Editor/BlueprintGraph/Private/EdGraphSchema_K2.cpp
@@ -2058,19 +2058,19 @@ const FPinConnectionResponse UEdGraphSchema_K2::DetermineConnectionResponseOfCom
 		!OutputPin->PinType.IsContainer() )
 	{
 		//check if the node wont be expanded as foreach call, if there is a link to an array
-		bool bAnyContainerInput = false;
+		bool bAnyArrayInput = false;
 		for(int InputLinkIndex = 0; InputLinkIndex < InputPin->LinkedTo.Num(); InputLinkIndex++)
 		{
 			if(const UEdGraphPin* Pin = InputPin->LinkedTo[InputLinkIndex])
 			{
-				if(Pin->PinType.IsContainer())
+				if(Pin->PinType.IsArray())
 				{
-					bAnyContainerInput = true;
+					bAnyArrayInput = true;
 					break;
 				}
 			}
 		}
-		bMultipleSelfException = !bAnyContainerInput;
+		bMultipleSelfException = !bAnyArrayInput;
 	}
 
 	if (bBreakExistingDueToExecOutput)
@@ -2209,12 +2209,16 @@ const FPinConnectionResponse UEdGraphSchema_K2::CanCreateConnection(const UEdGra
 	}
 
 	bool bIgnoreArray = false;
-	if(const UK2Node* OwningNode = Cast<UK2Node>(InputPin->GetOwningNode()))
+	if (const UK2Node* OwningNode = Cast<UK2Node>(InputPin->GetOwningNode()))
 	{
 		const bool bAllowMultipleSelfs = OwningNode->AllowMultipleSelfs(true); // it applies also to ForEachCall
 		const bool bNotAContainer = !InputPin->PinType.IsContainer();
 		const bool bSelfPin = IsSelfPin(*InputPin);
-		bIgnoreArray = bAllowMultipleSelfs && bNotAContainer && bSelfPin;
+		if (bAllowMultipleSelfs && bNotAContainer && bSelfPin)
+		{
+			// Indicates whether or not we will allow an array to be connected to a non-array input. This applies to nodes that support a foreach expansion of array inputs.
+			bIgnoreArray = OutputPin->PinType.IsArray();
+		}
 	}
 
 	// Find the calling context in case one of the pins is of type object and has a value of Self
