Fix loctext key collision
#rnx

diff --git a/Engine/Source/Editor/BlueprintGraph/Private/K2Node_CallFunction.cpp b/Engine/Source/Editor/BlueprintGraph/Private/K2Node_CallFunction.cpp
--- a/Engine/Source/Editor/BlueprintGraph/Private/K2Node_CallFunction.cpp
+++ b/Engine/Source/Editor/BlueprintGraph/Private/K2Node_CallFunction.cpp
@@ -1496,7 +1496,7 @@ FText UK2Node_CallFunction::GetTooltipText() const
 		// try to see where this function is meant to come from:
 		if (UClass* FuncOwnerClass = FunctionReference.GetMemberParentClass(GetBlueprintClassFromNode()))
 		{
-			return FText::Format(LOCTEXT("CallUnknownFunction", "Call unknown function {0} - missing from {1}"), FText::FromName(FunctionReference.GetMemberName()), FText::FromName(FuncOwnerClass->GetFName()));
+			return FText::Format(LOCTEXT("CallUnknownFunctionKnownOuter", "Call unknown function {0} - missing from {1}"), FText::FromName(FunctionReference.GetMemberName()), FText::FromName(FuncOwnerClass->GetFName()));
 		}
 		else
 		{
