Blueprint graph context menu perf improvements WiP (baseline).

Change summary:
- Redirected recently-added macro/postfix APIs (5.1) for wrapping filter rejection tests back to the old-style code paths as a baseline.
- Moved result caching into its own translation unit. We may be able to leverage this later for static context-based filters, so it remains in place as an internal API.
- Improved/hardened filter test profiling and item/stats logging. This remains disabled by default, but can be enabled locally for testing/debugging.
- Removed redundant EditorPtr context from FBlueprintActionMenuBuilder APIs as we already store it in the FBlueprintActionFilterContext structure.

diff --git a/Engine/Source/Editor/BlueprintGraph/Private/BlueprintActionFilter.cpp b/Engine/Source/Editor/BlueprintGraph/Private/BlueprintActionFilter.cpp
--- a/Engine/Source/Editor/BlueprintGraph/Private/BlueprintActionFilter.cpp
+++ b/Engine/Source/Editor/BlueprintGraph/Private/BlueprintActionFilter.cpp
@@ -42,6 +42,11 @@
 #include "Editor/EditorEngine.h"
 #include "BlueprintEditorSettings.h"
 
+#if ENABLE_BLUEPRINT_ACTION_FILTER_PROFILING
+#include "HAL/PlatformStackWalk.h"
+#include "ProfilingDebugging/CpuProfilerTrace.h"
+#endif	// ENABLE_BLUEPRINT_ACTION_FILTER_PROFILING
+
 /*******************************************************************************
  * Static BlueprintActionFilter Helpers
  ******************************************************************************/
@@ -454,6 +459,54 @@ namespace BlueprintActionFilterImpl
 	* @return True if the action is invalid to use in the current blueprint
 	*/
 	static bool IsHiddenInNonEditorBlueprint(FBlueprintActionFilter const& Filter, FBlueprintActionInfo& BlueprintAction);
+
+	//------------------------------------------------------------------------------
+#if ENABLE_BLUEPRINT_ACTION_FILTER_PROFILING
+	static TAutoConsoleVariable<bool> CVarBPEnableActionMenuFilterTestTraceLogging(
+		TEXT("BP.EnableActionMenuFilterTestTraceLogging"),
+		false,
+		TEXT("If enabled, each filter test will be individually profiled for Insights.")
+	);
+
+	static TAutoConsoleVariable<bool> CVarBPEnableActionMenuFilterTestStatsLogging(
+		TEXT("BP.EnableActionMenuFilterTestStatsLogging"),
+		false,
+		TEXT("If enabled, filter test stats will be logged for each menu section each time a Blueprint action menu is built for display.")
+	);
+
+	static FName FilterTestProfileEventName(TEXT("BlueprintActionFilterImpl::FilterTestProfileEvent"));
+
+	static const FString& FilterTestFuncPtrToFuncName(uint64 TestFuncPtr)
+	{
+		static TMap<uint64, FString> FuncPtrToFuncNameMap;
+
+		const FString* FuncNamePtr = FuncPtrToFuncNameMap.Find(TestFuncPtr);
+		if (!FuncNamePtr)
+		{
+			TRACE_CPUPROFILER_EVENT_SCOPE(FilterTestProfile_FuncPtrToFuncName);
+
+			FProgramCounterSymbolInfo SymbolInfo;
+			FPlatformStackWalk::ProgramCounterToSymbolInfo(TestFuncPtr, SymbolInfo);
+
+			FuncNamePtr = &FuncPtrToFuncNameMap.Add(TestFuncPtr, SymbolInfo.FunctionName);
+		}
+
+		return *FuncNamePtr;
+	}
+
+	static uint32 FilterTestFuncPtrToCpuSpecId(uint64 TestFuncPtr)
+	{
+		static TMap<uint64, uint32> FuncPtrToCpuSpecIdMap;
+
+		uint32* CpuSpecIdPtr = FuncPtrToCpuSpecIdMap.Find(TestFuncPtr);
+		if (!CpuSpecIdPtr)
+		{
+			CpuSpecIdPtr = &FuncPtrToCpuSpecIdMap.Add(TestFuncPtr, FCpuProfilerTrace::OutputEventType(*FilterTestFuncPtrToFuncName(TestFuncPtr)));
+		}
+
+		return *CpuSpecIdPtr;
+	}
+#endif	// ENABLE_BLUEPRINT_ACTION_FILTER_PROFILING
 };
 
 //------------------------------------------------------------------------------
@@ -2072,269 +2125,6 @@ UFunction const* FBlueprintActionInfo::GetAssociatedFunction()
 	return CachedActionFunction;
 }
 
-bool FBlueprintActionInfo::operator==(const FBlueprintActionInfo& other) const
-{
-	if (NodeSpawner != other.NodeSpawner || ActionOwner != other.ActionOwner)
-	{
-		return false;
-	}
-	
-	for (const FBindingObject &Binding : Bindings)
-	{
-		if (!other.Bindings.Contains(Binding))
-		{
-			return false;
-		}
-	}
-	return true;
-}
-
-/*******************************************************************************
- * FActionFilterCacheNode
- ******************************************************************************/
-
-uint32 GetTypeHash(const FBlueprintActionInfo& BlueprintActionInfo)
-{
-	uint32 Hash = GetTypeHash(BlueprintActionInfo.NodeSpawner);
-	Hash = HashCombine(Hash,GetTypeHash(BlueprintActionInfo.GetActionOwner()));
-	Hash = HashCombine(Hash,BlueprintActionInfo.GetBindings().Num());
-	
-	return Hash;
-}
-
-/*******************************************************************************
- * TActionFilterCacheKeyNode
- ******************************************************************************/
-
-template<typename TKeyType>
-class TActionFilterCacheKeyNode : public FActionFilterCacheNode
-{
-public:
-	explicit TActionFilterCacheKeyNode(const TKeyType& Key, TDelegate<bool(const TKeyType&, const TKeyType&)> Comparison = nullptr) :
-		Key(MakeShared<TKeyType>(Key)),
-		Comparison(Comparison)
-	{}
-	
-	virtual ~TActionFilterCacheKeyNode() override = default;
-
-	virtual void UnCacheSelf() override
-	{
-		if (Parent.IsValid())
-		{
-			const FSharedNode SharedParent = Parent.Pin();
-
-			// remove this node from parent
-			SharedParent->Children.RemoveSingleSwap(this->AsShared(), false);
-
-			if (SharedParent->Children.IsEmpty())
-			{
-				SharedParent->UnCacheSelf();
-			}
-		}
-	}
-	
-	TSharedPtr<TKeyType> Key = nullptr;
-protected:
-	virtual bool KeysMatch(const void* RawKey) const override
-	{
-		if (Comparison.IsBound())
-		{
-			return Comparison.Execute(*Key, *StaticCast<const TKeyType*>(RawKey));
-		}
-		return *Key == *StaticCast<const TKeyType*>(RawKey);
-	}
-
-	TDelegate<bool(const TKeyType&, const TKeyType&)> Comparison = nullptr;
-};
-
-/*******************************************************************************
- * TActionFilterCacheLeaf
- ******************************************************************************/
-
-template<typename TKeyType>
-class TActionFilterCacheLeaf : public TActionFilterCacheKeyNode<TKeyType>
-{
-public:
-	explicit TActionFilterCacheLeaf(const TKeyType& Key, TDelegate<bool(const TKeyType&, const TKeyType&)> Comparison = nullptr) :
-		TActionFilterCacheKeyNode<TKeyType>(Key, Comparison),
-		CachedFilterResults(MakeShared<TMap<FBlueprintActionInfo, bool>>())
-	{}
-	
-	virtual ~TActionFilterCacheLeaf() override = default;
-
-	virtual uint32 GetLeafAccessCount() const override
-	{
-		return AccessCount;
-	}
-
-
-	virtual void DecrementLeafAccessCount() override
-	{
-		if (AccessCount > 0)
-		{
-			--AccessCount;
-		}
-	}
-
-	uint32 AccessCount = 1;
-	TSharedPtr<TMap<FBlueprintActionInfo, bool>> CachedFilterResults;
-};
-
-/*******************************************************************************
- * FActionFilterCache
- ******************************************************************************/
-
-FActionFilterCache::FActionFilterCache()
-{
-	if (!ReloadCompleteDelegate.IsValid())
-	{
-		ReloadCompleteDelegate = FCoreUObjectDelegates::ReloadCompleteDelegate.AddStatic(&FActionFilterCache::OnReloadComplete);
-	}
-	if (!ModuleChangedDelegate.IsValid())
-	{
-		ModuleChangedDelegate = FModuleManager::Get().OnModulesChanged().AddStatic(&FActionFilterCache::OnModuleChanged);
-	}
-}
-
-void FActionFilterCache::OnReloadComplete(EReloadCompleteReason Reason)
-{
-	ClearAllCache();
-}
-
-void FActionFilterCache::OnModuleChanged(FName ModuleName, EModuleChangeReason ChangeReason)
-{
-	ClearAllCache();
-}
-
-void FActionFilterCache::ClearAllCache()
-{
-	SharedCache = MakeShared<FActionFilterCacheNode>();
-	CacheLeafs.Empty();
-}
-
-TSharedPtr<TMap<FBlueprintActionInfo, bool>> FActionFilterCache::GetCachedFilterResults(const FBlueprintActionFilter &Filter)
-{
-	// if the leaf for this context has already been looked up, just grab it rather than researching the tree
-	if (CachedFilterResults.IsValid())
-	{
-		return CachedFilterResults;
-	}
-	
-	TSharedPtr<FActionFilterCacheNode> Current = SharedCache;
-
-/** update Current to it's child that matches Key. Create one if neccesary */
-#define LOOKUP_CACHE_KEY(Key) \
-	FindOrCacheChildNode(Current, Key, #Key)
-	
-/** update Current to it's child that matches Key as a leaf. Create one if neccesary */
-#define LOOKUP_CACHE_KEY_LEAF(Key) \
-	FindOrCacheLeaf(Current, Key, #Key)
-
-/** update Current to it's child that matches Key as a leaf. Create one if neccesary
- *  uses Comparison lambda to check for key "equality"
- */
-#define LOOKUP_CACHE_KEY_PREDICATE(Key, Comparison) \
-	FindOrCacheChildNode<decltype(Key)>(Current, Key, #Key, TDelegate<bool(const decltype(Key)&, const decltype(Key)&)>::CreateLambda(Comparison))
-
-/** update Current to it's child that matches Key as a leaf. Create one if neccesary
- *  uses Comparison lambda to check for key "equality"
- */
-#define LOOKUP_CACHE_KEY_LEAF_PREDICATE(Key, Comparison) \
-	FindOrCacheLeaf<decltype(Key)>(Current, Key, #Key, TDelegate<bool(const decltype(Key)&, const decltype(Key)&)>::CreateLambda(Comparison))
-
-	// traverse/build the cache tree looking for the cached results
-
-	// cache based on what filter tests are being applied
-	LOOKUP_CACHE_KEY(Filter.FilterTests);
-
-	// cache based on which blueprints/graphs/editor we're generating actions for
-	LOOKUP_CACHE_KEY(Filter.Context.Blueprints);
-	LOOKUP_CACHE_KEY(Filter.Context.Graphs);
-	LOOKUP_CACHE_KEY(Filter.Context.EditorPtr);
-
-	// cache based on the filter's target classes
-	LOOKUP_CACHE_KEY(Filter.TargetClasses);
-
-	// cache based on what node types are permitted or rejected
-	LOOKUP_CACHE_KEY(Filter.RejectedNodeTypes);
-	LOOKUP_CACHE_KEY(Filter.PermittedNodeTypes);
-
-	// cache based on what flags are active
-	LOOKUP_CACHE_KEY(Filter.FilterFlags);
-	
-	// cache based on the relevant pin types and directions
-	LOOKUP_CACHE_KEY_LEAF_PREDICATE(Filter.Context.Pins, [](const TArray<UEdGraphPin*>& A, const TArray<UEdGraphPin*>& B)
-	{
-		if (A.Num() != B.Num())
-		{
-			return false;
-		}
-		for (int32 Index = 0; Index < A.Num(); ++Index)
-		{
-			if (A[Index]->Direction != B[Index]->Direction)
-			{
-				return false;
-			}
-			if (A[Index]->PinType != B[Index]->PinType)
-			{
-				return false;
-			}
-		}
-		return true;
-	});
-
-	// at this point, IsFilteredByThisCache has been updated to a shared memory space that can be retrieved again
-	// by following the above procedure with the same keys values.
-	return CachedFilterResults;
-}
-
-void FActionFilterCache::PopCache()
-{
-	int32 MinIndex = 0;
-	uint32 MinAccessCount = TNumericLimits<uint32>::Max();
-	for (int32 LeafIndex = 0; LeafIndex < CacheLeafs.Num(); ++LeafIndex)
-	{
-		TWeakPtr<FActionFilterCacheNode>& Leaf = CacheLeafs[LeafIndex];
-		// if there's invalid cache, remove it first
-		if (!Leaf.IsValid())
-		{
-			MinIndex = LeafIndex;
-			break;
-		}
-
-		const uint32 AccessCount = Leaf.Pin()->GetLeafAccessCount();
-		if (AccessCount < MinAccessCount)
-		{
-			MinAccessCount = AccessCount;
-			MinIndex = LeafIndex;
-		}
-
-		// decrement access count of all leafs so that recently used leafs get prioritized more
-		Leaf.Pin()->DecrementLeafAccessCount();
-	}
-
-	TWeakPtr<FActionFilterCacheNode> RemovedElement = CacheLeafs[MinIndex];
-	if (RemovedElement.IsValid())
-	{
-		RemovedElement.Pin()->UnCacheSelf();
-	}
-
-	CacheLeafs.RemoveAt(MinIndex, 1, false);
-}
-
-TSharedPtr<FActionFilterCacheNode> FActionFilterCache::SharedCache = MakeShared<FActionFilterCacheNode>();
-
-TAutoConsoleVariable<int32> FActionFilterCache::CVarCacheLeafCapacity{
-	TEXT("BP.ActionMenuFilterCacheLeafCapacity"),
-	32,
-	TEXT("The number of action menu contexts to cache simultaniously. raising this number will increase the memory footprint but decrease how often the cache is blown")
-};
-
-TArray<TWeakPtr<FActionFilterCacheNode>> FActionFilterCache::CacheLeafs;
-
-FDelegateHandle FActionFilterCache::ReloadCompleteDelegate;
-FDelegateHandle FActionFilterCache::ModuleChangedDelegate;
-
 /*******************************************************************************
  * FActionFilterTest
  ******************************************************************************/
@@ -2352,31 +2142,11 @@ FActionFilterTest::FActionFilterTest(const FRejectionDelegate Delegate, FString
 {
 }
 
-bool FActionFilterTest::Call(FBlueprintActionFilter const& Filter, FBlueprintActionInfo& ActionInfo,
-                             TMap<FString, uint32>* FilterTestToCpuSpecIds)
+bool FActionFilterTest::Call(FBlueprintActionFilter const& Filter, FBlueprintActionInfo& ActionInfo)
 {
 	check(RejectionDelegate.IsBound());
 
-#if ENABLE_BLUEPRINT_ACTION_FILTER_PROFILING
-	checkf(FilterTestToCpuSpecIds, TEXT("If Action Filter Profiling is enabled, FilterTestToCpuSpecIds must be provided"));
-	uint32& CpuSpecId = FilterTestToCpuSpecIds->FindOrAdd(Name);
-	FCpuProfilerTrace::FEventScope ProfilerEvent(CpuSpecId, *FString::Printf(TEXT("IsFiltered_%s"), *Name), true, __FILE__, __LINE__);
-
-	const double StartTime = FPlatformTime::Seconds();
-#endif	// ENABLE_BLUEPRINT_ACTION_FILTER_PROFILING
-
-	const bool bIsFiltered = RejectionDelegate.Execute(Filter, ActionInfo);
-
-#if ENABLE_BLUEPRINT_ACTION_FILTER_PROFILING
-	ProfileRecord.NumIterations += 1;
-	ProfileRecord.TotalTimeMs += (FPlatformTime::Seconds() - StartTime) * 1000.0f;
-	if (bIsFiltered)
-	{
-		ProfileRecord.NumFilteredOut += 1;
-	}
-#endif	// ENABLE_BLUEPRINT_ACTION_FILTER_PROFILING
-	
-	return bIsFiltered;
+	return RejectionDelegate.Execute(Filter, ActionInfo);
 }
 
 bool FActionFilterTest::operator==(const FActionFilterTest& other) const
@@ -2391,17 +2161,16 @@ bool FActionFilterTest::operator==(const FActionFilterTest& other) const
 //------------------------------------------------------------------------------
 // DEPRECATED
 FBlueprintActionFilter::FBlueprintActionFilter(uint32 Flags)
-	:FBlueprintActionFilter(EFlags(Flags))
+	: FBlueprintActionFilter(EFlags(Flags))
 {
 }
 
 //------------------------------------------------------------------------------
 FBlueprintActionFilter::FBlueprintActionFilter(const EFlags InFlags /*= BPFILTER_NoFlags*/)
+	: FilterFlags(InFlags)
 {
 	using namespace BlueprintActionFilterImpl;
 
-	FilterFlags = InFlags;
-
 	// Disable non-imported type filtering if namespace filtering features are disabled.
 	if (!GetDefault<UBlueprintEditorSettings>()->bEnableNamespaceFilteringFeatures)
 	{
@@ -2422,59 +2191,55 @@ FBlueprintActionFilter::FBlueprintActionFilter(const EFlags InFlags /*= BPFILTER
 	//
 	// this test in-particular spawns a template-node and then calls 
 	// AllocateDefaultPins() which is costly, so it should be very last!
-	AddRejectionTest(MAKE_ACTION_FILTER_REJECTION_TEST(IsNodeTemplateSelfFiltered)); // since this is virtual, we'll avoid caching it.
-	AddRejectionTest(MAKE_ACTION_FILTER_REJECTION_TEST(IsIncompatibleAnimNotification)->WithFlags(EActionFilterTestFlags::CacheResults));
-	AddRejectionTest(MAKE_ACTION_FILTER_REJECTION_TEST(IsAssetPermissionNotGranted)); // since this depends on other filters, we'll avoid caching it
-
-	// TODO: blow cache when the results of the following might change:
-	//   - K2Schema->FunctionHasParamOfType
-	//   - K2Schema->ArePinTypesCompatible
-	//   - K2Schema->ConvertPropertyToPinType
-	AddRejectionTest(MAKE_ACTION_FILTER_REJECTION_TEST(IsMissingMatchingPinParam)->WithFlags(EActionFilterTestFlags::CacheResults));
-	AddRejectionTest(MAKE_ACTION_FILTER_REJECTION_TEST(IsMissmatchedPropertyType)->WithFlags(EActionFilterTestFlags::CacheResults));
-	AddRejectionTest(MAKE_ACTION_FILTER_REJECTION_TEST(IsFunctionMissingPinParam)->WithFlags(EActionFilterTestFlags::CacheResults));
-	AddRejectionTest(MAKE_ACTION_FILTER_REJECTION_TEST(IsIncompatibleLatentNode)->WithFlags(EActionFilterTestFlags::CacheResults));
-	AddRejectionTest(MAKE_ACTION_FILTER_REJECTION_TEST(IsIncompatibleImpureNode)->WithFlags(EActionFilterTestFlags::CacheResults));
-	AddRejectionTest(MAKE_ACTION_FILTER_REJECTION_TEST(IsPropertyAccessorNode)->WithFlags(EActionFilterTestFlags::CacheResults));
+	AddRejectionTest(FRejectionTestDelegate::CreateStatic(IsNodeTemplateSelfFiltered));
+
+	AddRejectionTest(FRejectionTestDelegate::CreateStatic(IsIncompatibleAnimNotification));
+	AddRejectionTest(FRejectionTestDelegate::CreateStatic(IsAssetPermissionNotGranted));
+	AddRejectionTest(FRejectionTestDelegate::CreateStatic(IsMissingMatchingPinParam));
+	AddRejectionTest(FRejectionTestDelegate::CreateStatic(IsMissmatchedPropertyType));
+	AddRejectionTest(FRejectionTestDelegate::CreateStatic(IsFunctionMissingPinParam));
+	AddRejectionTest(FRejectionTestDelegate::CreateStatic(IsIncompatibleLatentNode));
+	AddRejectionTest(FRejectionTestDelegate::CreateStatic(IsIncompatibleImpureNode));
+	AddRejectionTest(FRejectionTestDelegate::CreateStatic(IsPropertyAccessorNode));
+
 	if (HasAnyFlags(BPFILTER_RejectIncompatibleThreadSafety))
 	{
-		AddRejectionTest(MAKE_ACTION_FILTER_REJECTION_TEST(IsThreadSafetyIncompatible)->WithFlags(EActionFilterTestFlags::CacheResults));
+		AddRejectionTest(FRejectionTestDelegate::CreateStatic(IsThreadSafetyIncompatible));
 	}
-	
-	AddRejectionTest(MAKE_ACTION_FILTER_REJECTION_TEST(IsActionHiddenByConfig)->WithFlags(EActionFilterTestFlags::CacheResults));
-	AddRejectionTest(MAKE_ACTION_FILTER_REJECTION_TEST(IsFieldCategoryHidden)->WithFlags(EActionFilterTestFlags::CacheResults));
+
+	AddRejectionTest(FRejectionTestDelegate::CreateStatic(IsActionHiddenByConfig));
+	AddRejectionTest(FRejectionTestDelegate::CreateStatic(IsFieldCategoryHidden));
+
 	if (HasAnyFlags(BPFILTER_RejectGlobalFields | BPFILTER_RejectNonImportedFields))
 	{
-		AddRejectionTest(MAKE_ACTION_FILTER_REJECTION_TEST(IsRejectedGlobalField)->WithFlags(EActionFilterTestFlags::CacheResults));
+		AddRejectionTest(FRejectionTestDelegate::CreateStatic(IsRejectedGlobalField));
 	}
-	
-	AddRejectionTest(MAKE_ACTION_FILTER_REJECTION_TEST(IsFieldInaccessible)->WithFlags(EActionFilterTestFlags::CacheResults));
-	AddRejectionTest(MAKE_ACTION_FILTER_REJECTION_TEST(IsNotSubClassCast)->WithFlags(EActionFilterTestFlags::CacheResults));
-	AddRejectionTest(MAKE_ACTION_FILTER_REJECTION_TEST(IsEventUnimplementable)->WithFlags(EActionFilterTestFlags::CacheResults));
-	AddRejectionTest(MAKE_ACTION_FILTER_REJECTION_TEST(IsPermissionNotGranted)->WithFlags(EActionFilterTestFlags::CacheResults));
-	AddRejectionTest(MAKE_ACTION_FILTER_REJECTION_TEST(IsRestrictedClassMember)->WithFlags(EActionFilterTestFlags::CacheResults));
-	AddRejectionTest(MAKE_ACTION_FILTER_REJECTION_TEST(IsIncompatibleWithGraphType)->WithFlags(EActionFilterTestFlags::CacheResults));
-	AddRejectionTest(MAKE_ACTION_FILTER_REJECTION_TEST(IsSchemaIncompatible)->WithFlags(EActionFilterTestFlags::CacheResults));
-	AddRejectionTest(MAKE_ACTION_FILTER_REJECTION_TEST(IsExtraneousInterfaceCall)->WithFlags(EActionFilterTestFlags::CacheResults));
-	AddRejectionTest(MAKE_ACTION_FILTER_REJECTION_TEST(IsIncompatibleMacroInstance)->WithFlags(EActionFilterTestFlags::CacheResults));
+
+	AddRejectionTest(FRejectionTestDelegate::CreateStatic(IsFieldInaccessible));
+	AddRejectionTest(FRejectionTestDelegate::CreateStatic(IsNotSubClassCast));
+	AddRejectionTest(FRejectionTestDelegate::CreateStatic(IsEventUnimplementable));
+	AddRejectionTest(FRejectionTestDelegate::CreateStatic(IsPermissionNotGranted));
+	AddRejectionTest(FRejectionTestDelegate::CreateStatic(IsRestrictedClassMember));
+	AddRejectionTest(FRejectionTestDelegate::CreateStatic(IsIncompatibleWithGraphType));
+	AddRejectionTest(FRejectionTestDelegate::CreateStatic(IsSchemaIncompatible));
+	AddRejectionTest(FRejectionTestDelegate::CreateStatic(IsExtraneousInterfaceCall));
+	AddRejectionTest(FRejectionTestDelegate::CreateStatic(IsIncompatibleMacroInstance));
 
 	if (!HasAnyFlags(BPFILTER_PermitDeprecated))
 	{
-		AddRejectionTest(MAKE_ACTION_FILTER_REJECTION_TEST(IsDeprecated)->WithFlags(EActionFilterTestFlags::CacheResults));
+		AddRejectionTest(FRejectionTestDelegate::CreateStatic(IsDeprecated));
 	}
 
-	AddRejectionTest(MAKE_ACTION_FILTER_REJECTION_TEST(IsFilteredNodeType)->WithFlags(EActionFilterTestFlags::CacheResults));
-	AddRejectionTest(MAKE_ACTION_FILTER_REJECTION_TEST(IsNonTargetMember)->WithFlags(EActionFilterTestFlags::CacheResults));
-	AddRejectionTest(MAKE_ACTION_FILTER_REJECTION_TEST(IsUnBoundBindingSpawner)->WithFlags(EActionFilterTestFlags::CacheResults));
-	AddRejectionTest(MAKE_ACTION_FILTER_REJECTION_TEST(IsOutOfScopeLocalVariable)->WithFlags(EActionFilterTestFlags::CacheResults));
-	AddRejectionTest(MAKE_ACTION_FILTER_REJECTION_TEST(IsLevelScriptActionValid)->WithFlags(EActionFilterTestFlags::CacheResults));
-
-	AddRejectionTest(MAKE_ACTION_FILTER_REJECTION_TEST(IsHiddenInNonEditorBlueprint)->WithFlags(EActionFilterTestFlags::CacheResults));
-
+	AddRejectionTest(FRejectionTestDelegate::CreateStatic(IsFilteredNodeType));
+	AddRejectionTest(FRejectionTestDelegate::CreateStatic(IsNonTargetMember));
+	AddRejectionTest(FRejectionTestDelegate::CreateStatic(IsUnBoundBindingSpawner));
+	AddRejectionTest(FRejectionTestDelegate::CreateStatic(IsOutOfScopeLocalVariable));
+	AddRejectionTest(FRejectionTestDelegate::CreateStatic(IsLevelScriptActionValid));
+	AddRejectionTest(FRejectionTestDelegate::CreateStatic(IsHiddenInNonEditorBlueprint));
 
 	// added as the first rejection test, so that we don't operate on stale 
 	// (TRASH/REINST) class fields
-	AddRejectionTest(MAKE_ACTION_FILTER_REJECTION_TEST(IsStaleFieldAction));
+	AddRejectionTest(FRejectionTestDelegate::CreateStatic(IsStaleFieldAction));
 }
 
 //------------------------------------------------------------------------------
@@ -2499,20 +2264,32 @@ void FBlueprintActionFilter::Add(TArray<FTargetClassFilterData>& ToArray, UClass
 	ToArray.Add(Data);
 }
 
+//------------------------------------------------------------------------------
+void FBlueprintActionFilter::AddRejectionTest(FRejectionTestDelegate RejectionTestDelegate)
+{
+	if (ensureMsgf(RejectionTestDelegate.IsBound(), TEXT("Cannot add a rejection test without a bound delegate")))
+	{
+		FilterTests.Add(RejectionTestDelegate);
+
+#if ENABLE_BLUEPRINT_ACTION_FILTER_PROFILING
+		FFilterTestProfileRecord FilterTestProfileRecord;
+		FilterTestProfileRecord.TestFuncPtr = RejectionTestDelegate.GetBoundProgramCounterForTimerManager();
+
+		FilterTestProfiles.Add(FilterTests.Num() - 1, MoveTemp(FilterTestProfileRecord));
+#endif	// ENABLE_BLUEPRINT_ACTION_FILTER_PROFILING
+	}
+}
+
 //------------------------------------------------------------------------------
 void FBlueprintActionFilter::AddRejectionTest(TSharedRef<FActionFilterTest> RejectionTest)
 {
-	checkf(RejectionTest->RejectionDelegate.IsBound(), TEXT("Cannot add a rejection test without a bound delegate"));
-	FilterTests.Add(RejectionTest);
+	AddRejectionTest(RejectionTest->RejectionDelegate);
 }
 
 //------------------------------------------------------------------------------
 bool FBlueprintActionFilter::IsFiltered(FBlueprintActionInfo& BlueprintAction)
 {
-	
-#if ENABLE_BLUEPRINT_ACTION_FILTER_PROFILING
 	TRACE_CPUPROFILER_EVENT_SCOPE(FBlueprintActionFilter::IsFiltered);
-#endif
 
 	bool bIsFiltered = IsFilteredByThis(BlueprintAction);
 	if (!bIsFiltered)
@@ -2553,24 +2330,10 @@ FBlueprintActionFilter const& FBlueprintActionFilter::operator&=(FBlueprintActio
 }
 
 //------------------------------------------------------------------------------
-TSharedPtr<TMap<FBlueprintActionInfo, bool>> FBlueprintActionFilter::GetCachedFilterResults()
-{
-	return ActionFilterCache.GetCachedFilterResults(*this);
-}
-
-const TArray<TSharedRef<FActionFilterTest>> &FBlueprintActionFilter::GetFilterTests() const
+bool FBlueprintActionFilter::IsFilteredByThis(FBlueprintActionInfo& BlueprintAction) const
 {
-	return FilterTests;
-}
+	TRACE_CPUPROFILER_EVENT_SCOPE(FBlueprintActionFilter::IsFilteredByThis);
 
-//------------------------------------------------------------------------------
-bool FBlueprintActionFilter::IsFilteredByThis(FBlueprintActionInfo& BlueprintAction)
-{
-	if (IsFilteredByCachedFilters(BlueprintAction))
-	{
-		return true;
-	}
-	
 	FBlueprintActionFilter const& FilterRef = *this;
 
 	// for debugging purposes:
@@ -2582,28 +2345,39 @@ bool FBlueprintActionFilter::IsFilteredByThis(FBlueprintActionInfo& BlueprintAct
 // 	}
 
 #if ENABLE_BLUEPRINT_ACTION_FILTER_PROFILING
-	static TMap<FString, uint32> FilterTestToCpuSpecIds;
+	const bool bIsFilterTestProfilingEnabled = IsFilterTestTraceLoggingEnabled() || IsFilterTestStatsLoggingEnabled();
 #endif	// ENABLE_BLUEPRINT_ACTION_FILTER_PROFILING
 
 	bool bIsFiltered = false;
 	// iterate backwards so that custom user test are ran first (and the slow
 	// internal tests are ran last).
-	for (int32 TestIndex = FilterTests.Num()-1; TestIndex >= 0; --TestIndex)
+	for (int32 TestIndex = FilterTests.Num() - 1; TestIndex >= 0; --TestIndex)
 	{
-		TSharedRef<FActionFilterTest> RejectionTest = FilterTests[TestIndex];
+		FRejectionTestDelegate const& RejectionTestDelegate = FilterTests[TestIndex];
+		checkSlow(RejectionTestDelegate.IsBound());
 
-		// cached tests have already been ran. we can skip them.
-		if (EnumHasAllFlags(RejectionTest->Flags, EActionFilterTestFlags::CacheResults))
+#if ENABLE_BLUEPRINT_ACTION_FILTER_PROFILING
+		double StartTime = 0;
+		FFilterTestProfileRecord* FilterTestProfileRecord = nullptr;
+		if (bIsFilterTestProfilingEnabled)
 		{
-			continue;
+			FCpuProfilerTrace::OutputBeginDynamicEvent(BlueprintActionFilterImpl::FilterTestProfileEventName);
+
+			StartTime = FPlatformTime::Seconds();
+			FilterTestProfileRecord = BeginFilterTestProfileEvent(TestIndex, RejectionTestDelegate.GetBoundProgramCounterForTimerManager());
 		}
+#endif	// ENABLE_BLUEPRINT_ACTION_FILTER_PROFILING
+
+		bIsFiltered = RejectionTestDelegate.Execute(FilterRef, BlueprintAction);
 
 #if ENABLE_BLUEPRINT_ACTION_FILTER_PROFILING
-		bIsFiltered = RejectionTest.Call(FilterRef, BlueprintAction, &FilterTestToCpuSpecIds);
-#else
-		bIsFiltered = RejectionTest->Call(FilterRef, BlueprintAction);
-#endif
-		
+		if (FilterTestProfileRecord)
+		{
+			EndFilterTestProfileEvent(FilterTestProfileRecord, bIsFiltered, StartTime);
+			FCpuProfilerTrace::OutputEndEvent();
+		}
+#endif	// ENABLE_BLUEPRINT_ACTION_FILTER_PROFILING
+
 		if (bIsFiltered)
 		{
 			break;
@@ -2612,9 +2386,34 @@ bool FBlueprintActionFilter::IsFilteredByThis(FBlueprintActionInfo& BlueprintAct
 
 	if (!bIsFiltered)
 	{
+#if ENABLE_BLUEPRINT_ACTION_FILTER_PROFILING
+		int32 ExtraTestProfileRecordIndex = FilterTests.Num();
+#endif	// ENABLE_BLUEPRINT_ACTION_FILTER_PROFILING
+
 		for (const FBlueprintGraphModule::FActionMenuRejectionTest& ExtraRejectionTest : BluprintGraphModule->GetExtendedActionMenuFilters())
 		{
+#if ENABLE_BLUEPRINT_ACTION_FILTER_PROFILING
+			double StartTime = 0;
+			FFilterTestProfileRecord* FilterTestProfileRecord = nullptr;
+			if (bIsFilterTestProfilingEnabled)
+			{
+				FCpuProfilerTrace::OutputBeginDynamicEvent(BlueprintActionFilterImpl::FilterTestProfileEventName);
+
+				StartTime = FPlatformTime::Seconds();
+				FilterTestProfileRecord = BeginFilterTestProfileEvent(ExtraTestProfileRecordIndex++, ExtraRejectionTest.GetBoundProgramCounterForTimerManager());
+			}
+#endif	// ENABLE_BLUEPRINT_ACTION_FILTER_PROFILING
+
 			bIsFiltered = ExtraRejectionTest.Execute(FilterRef, BlueprintAction);
+
+#if ENABLE_BLUEPRINT_ACTION_FILTER_PROFILING
+			if (FilterTestProfileRecord)
+			{
+				EndFilterTestProfileEvent(FilterTestProfileRecord, bIsFiltered, StartTime);
+				FCpuProfilerTrace::OutputEndEvent();
+			}
+#endif	// ENABLE_BLUEPRINT_ACTION_FILTER_PROFILING
+
 			if (bIsFiltered)
 			{
 				break;
@@ -2625,95 +2424,83 @@ bool FBlueprintActionFilter::IsFilteredByThis(FBlueprintActionInfo& BlueprintAct
 	return bIsFiltered;
 }
 
-bool FBlueprintActionFilter::IsFilteredByCachedFilters(FBlueprintActionInfo& BlueprintAction)
-{
-	
+//------------------------------------------------------------------------------
 #if ENABLE_BLUEPRINT_ACTION_FILTER_PROFILING
-	TRACE_CPUPROFILER_EVENT_SCOPE(FBlueprintActionFilter::IsFilteredByCachedFilters);
-#endif
-	
-	const TSharedPtr<TMap<FBlueprintActionInfo, bool>> Cache = GetCachedFilterResults();
-	FBlueprintActionFilter const& FilterRef = *this;
-		
-	if (BlueprintActionFilterImpl::IsStaleFieldAction(FilterRef, BlueprintAction))
+bool FBlueprintActionFilter::IsFilterTestTraceLoggingEnabled()
+{
+	return BlueprintActionFilterImpl::CVarBPEnableActionMenuFilterTestTraceLogging.GetValueOnGameThread();
+}
+
+bool FBlueprintActionFilter::IsFilterTestStatsLoggingEnabled()
+{
+	return BlueprintActionFilterImpl::CVarBPEnableActionMenuFilterTestStatsLogging.GetValueOnGameThread();
+}
+
+TArray<FString> FBlueprintActionFilter::GetFilterTestProfile()
+{
+	TArray<FString> Result;
+
+	auto AddFilterTestProfile = [&Result](int32 DisplayIndex, const FFilterTestProfileRecord& FilterTestProfile)
 	{
-		// if BlueprintAction is stale, immediately filter the action
-		if (FActionFilterTest::CVarEnableCaching.GetValueOnGameThread())
-		{
-			Cache->FindOrAdd(BlueprintAction) = true;
-		}
-		return true;
+		const FString& FilterTestFunctionName = BlueprintActionFilterImpl::FilterTestFuncPtrToFuncName(FilterTestProfile.TestFuncPtr);
+		Result.Add(FString::Printf(TEXT("[%02d]: %s | %d tested | %d failed | %.02f ms | %.02f us (per test)"),
+			DisplayIndex,
+			*FilterTestFunctionName,
+			FilterTestProfile.NumIterations,
+			FilterTestProfile.NumFilteredOut,
+			FilterTestProfile.TotalTimeMs,
+			FilterTestProfile.NumIterations > 0 ? FilterTestProfile.TotalTimeMs / FilterTestProfile.NumIterations * 1000.0f : 0.0f)
+		);
+	};
+
+	for (int32 FilterTestIndex = FilterTests.Num() - 1; FilterTestIndex >= 0; --FilterTestIndex)
+	{
+		AddFilterTestProfile(FilterTests.Num() - FilterTestIndex, FilterTestProfiles[FilterTestIndex]);
 	}
 
-	if (FActionFilterTest::CVarEnableCaching.GetValueOnGameThread())
+	for (int32 ExtendedFilterTestIndex = FilterTests.Num(); ExtendedFilterTestIndex < FilterTests.Num() + BluprintGraphModule->GetExtendedActionMenuFilters().Num(); ++ExtendedFilterTestIndex)
 	{
-		if (const bool* Result = Cache->Find(BlueprintAction))
-		{
-			return *Result;
-		}
+		AddFilterTestProfile(ExtendedFilterTestIndex, FilterTestProfiles[ExtendedFilterTestIndex]);
 	}
 
-#if ENABLE_BLUEPRINT_ACTION_FILTER_PROFILING
-	static TMap<FString, uint32> FilterTestToCpuSpecIds;
-#endif	// ENABLE_BLUEPRINT_ACTION_FILTER_PROFILING
-	
-	bool bIsFiltered = false;
-	// iterate backwards so that custom user test are ran first (and the slow
-	// internal tests are ran last).
-	for (int32 TestIndex = FilterTests.Num()-1; TestIndex >= 0; --TestIndex)
+	return Result;
+}
+
+FBlueprintActionFilter::FFilterTestProfileRecord* FBlueprintActionFilter::BeginFilterTestProfileEvent(int32 FilterTestIndex, uint64 FilterTestFuncPtr) const
+{
+	FFilterTestProfileRecord* FilterTestProfileRecord = FilterTestProfiles.Find(FilterTestIndex);
+	if (!FilterTestProfileRecord)
 	{
-		TSharedRef<FActionFilterTest> RejectionTest = FilterTests[TestIndex];
-		
-		// this method only checks cached rejection tests
-		if (!EnumHasAllFlags(RejectionTest->Flags, EActionFilterTestFlags::CacheResults))
-		{
-			continue;
-		}
-		
-#if ENABLE_BLUEPRINT_ACTION_FILTER_PROFILING
-		bIsFiltered = RejectionTest.Call(FilterRef, BlueprintAction, &FilterTestToCpuSpecIds);
-#else
-		bIsFiltered = RejectionTest->Call(FilterRef, BlueprintAction);
-#endif
-		
-		if (bIsFiltered)
-		{
-			break;
-		}
+		FilterTestProfileRecord = &FilterTestProfiles.Add(FilterTestIndex);
+		FilterTestProfileRecord->TestFuncPtr = FilterTestFuncPtr;
 	}
-	
-	if (FActionFilterTest::CVarEnableCaching.GetValueOnGameThread())
+
+	if (IsFilterTestTraceLoggingEnabled())
 	{
-		Cache->Add(BlueprintAction, bIsFiltered);
+		FCpuProfilerTrace::OutputBeginEvent(BlueprintActionFilterImpl::FilterTestFuncPtrToCpuSpecId(FilterTestFuncPtr));
 	}
-	
-	return bIsFiltered;
+
+	return FilterTestProfileRecord;
 }
 
-//------------------------------------------------------------------------------
-#if ENABLE_BLUEPRINT_ACTION_FILTER_PROFILING
-TArray<FString> FBlueprintActionFilter::GetFilterTestProfile()
+void FBlueprintActionFilter::EndFilterTestProfileEvent(FFilterTestProfileRecord* InEvent, bool bWasFiltered, double StartTime) const
 {
-	TArray<FString> Result;
-
-	for (int32 TestIndex = FilterTests.Num() - 1; TestIndex >= 0; --TestIndex)
+	if (InEvent)
 	{
-		FActionFilterTest& RejectionTest = FilterTests[TestIndex];
-		FActionFilterTest::FFilterTestProfileRecord& ProfileRecord = RejectionTest.ProfileRecord;
-		
-		Result.Add(FString::Printf(TEXT("[%02d]: %s | %d tested | %d failed | %.02f ms | %.02f us (per test)"),
-			FilterTests.Num() - TestIndex,
-			*RejectionTest.Name,
-			ProfileRecord.NumIterations,
-			ProfileRecord.NumFilteredOut,
-			ProfileRecord.TotalTimeMs,
-			ProfileRecord.TotalTimeMs / ProfileRecord.NumIterations * 1000.0f
-		));
+		if (IsFilterTestTraceLoggingEnabled())
+		{
+			FCpuProfilerTrace::OutputEndEvent();
+		}
 
-		// reset record
-		ProfileRecord = FActionFilterTest::FFilterTestProfileRecord();
+		if (IsFilterTestStatsLoggingEnabled())
+		{
+			InEvent->NumIterations += 1;
+			InEvent->TotalTimeMs += (FPlatformTime::Seconds() - StartTime) * 1000.0f;
+			if (bWasFiltered)
+			{
+				InEvent->NumFilteredOut += 1;
+			}
+		}
 	}
-
-	return Result;
 }
 #endif	// ENABLE_BLUEPRINT_ACTION_FILTER_PROFILING
\ No newline at end of file
diff --git a/Engine/Source/Editor/BlueprintGraph/Private/BlueprintActionFilterResultCache.cpp b/Engine/Source/Editor/BlueprintGraph/Private/BlueprintActionFilterResultCache.cpp
new file mode 100644
--- /dev/null
+++ b/Engine/Source/Editor/BlueprintGraph/Private/BlueprintActionFilterResultCache.cpp
@@ -0,0 +1,279 @@
+// Copyright Epic Games, Inc. All Rights Reserved.
+
+#include "BlueprintActionFilterResultCache.h"
+#include "Concepts/EqualityComparable.h"
+
+/*******************************************************************************
+ * FActionFilterCacheKey
+ ******************************************************************************/
+
+FActionFilterCacheKey::FActionFilterCacheKey(const FBlueprintActionInfo& InAction)
+{
+	KeyHash = GetTypeHash(InAction.NodeSpawner);
+	KeyHash = HashCombine(KeyHash, GetTypeHash(InAction.GetActionOwner()));
+	KeyHash = HashCombine(KeyHash, InAction.GetBindings().Num());
+}
+
+bool operator==(const FActionFilterCacheKey& A, const FActionFilterCacheKey& B)
+{
+	return A.KeyHash == B.KeyHash;
+}
+
+uint32 GetTypeHash(const FActionFilterCacheKey& InKey)
+{
+	return InKey.KeyHash;
+}
+
+/*******************************************************************************
+ * TActionFilterCacheKeyNode
+ ******************************************************************************/
+
+template<typename TKeyType>
+class TActionFilterCacheKeyNode : public FActionFilterCacheNode
+{
+public:
+	explicit TActionFilterCacheKeyNode(const TKeyType& Key, TDelegate<bool(const TKeyType&, const TKeyType&)> Comparison = nullptr) :
+		Key(MakeShared<TKeyType>(Key)),
+		Comparison(Comparison)
+	{}
+
+	virtual ~TActionFilterCacheKeyNode() override = default;
+
+	virtual void UnCacheSelf() override
+	{
+		if (Parent.IsValid())
+		{
+			const FSharedNode SharedParent = Parent.Pin();
+
+			// remove this node from parent
+			SharedParent->Children.RemoveSingleSwap(this->AsShared(), false);
+
+			if (SharedParent->Children.IsEmpty())
+			{
+				SharedParent->UnCacheSelf();
+			}
+		}
+	}
+
+	TSharedPtr<TKeyType> Key = nullptr;
+protected:
+	virtual bool KeysMatch(const void* RawKey) const override
+	{
+		if (Comparison.IsBound())
+		{
+			return Comparison.Execute(*Key, *StaticCast<const TKeyType*>(RawKey));
+		}
+		else
+		{
+			return *Key == *StaticCast<const TKeyType*>(RawKey);
+		}
+	}
+
+	TDelegate<bool(const TKeyType&, const TKeyType&)> Comparison = nullptr;
+};
+
+/*******************************************************************************
+ * TActionFilterCacheLeaf
+ ******************************************************************************/
+
+template<typename TKeyType>
+class TActionFilterCacheLeaf : public TActionFilterCacheKeyNode<TKeyType>
+{
+public:
+	explicit TActionFilterCacheLeaf(const TKeyType& Key, TDelegate<bool(const TKeyType&, const TKeyType&)> Comparison = nullptr) :
+		TActionFilterCacheKeyNode<TKeyType>(Key, Comparison),
+		CachedFilterResults(MakeShared<TMap<FActionFilterCacheKey, bool>>())
+	{}
+
+	virtual ~TActionFilterCacheLeaf() override = default;
+
+	virtual uint32 GetLeafAccessCount() const override
+	{
+		return AccessCount;
+	}
+
+
+	virtual void DecrementLeafAccessCount() override
+	{
+		if (AccessCount > 0)
+		{
+			--AccessCount;
+		}
+	}
+
+	uint32 AccessCount = 1;
+	TSharedPtr<TMap<FActionFilterCacheKey, bool>> CachedFilterResults;
+};
+
+/*******************************************************************************
+ * FActionFilterCache
+ ******************************************************************************/
+
+FActionFilterCache::FActionFilterCache()
+{
+	if (!ReloadCompleteDelegate.IsValid())
+	{
+		ReloadCompleteDelegate = FCoreUObjectDelegates::ReloadCompleteDelegate.AddStatic(&FActionFilterCache::OnReloadComplete);
+	}
+	if (!ModuleChangedDelegate.IsValid())
+	{
+		ModuleChangedDelegate = FModuleManager::Get().OnModulesChanged().AddStatic(&FActionFilterCache::OnModuleChanged);
+	}
+}
+
+void FActionFilterCache::OnReloadComplete(EReloadCompleteReason Reason)
+{
+	ClearAllCache();
+}
+
+void FActionFilterCache::OnModuleChanged(FName ModuleName, EModuleChangeReason ChangeReason)
+{
+	ClearAllCache();
+}
+
+void FActionFilterCache::ClearAllCache()
+{
+	SharedCache = MakeShared<FActionFilterCacheNode>();
+	CacheLeafs.Empty();
+}
+
+void FActionFilterCache::CacheFilteredResult(const FBlueprintActionFilter& Filter, FBlueprintActionInfo& BlueprintAction, bool bWasFilteredOut)
+{
+	const TSharedPtr<TMap<FActionFilterCacheKey, bool>> Cache = GetCachedFilterResults(Filter);
+	Cache->Add(FActionFilterCacheKey(BlueprintAction), bWasFilteredOut);
+}
+
+bool FActionFilterCache::IsFilteredByCachedResult(const FBlueprintActionFilter& Filter, FBlueprintActionInfo& BlueprintAction)
+{
+	const TSharedPtr<TMap<FActionFilterCacheKey, bool>> Cache = GetCachedFilterResults(Filter);
+	if (const bool* Result = Cache->Find(FActionFilterCacheKey(BlueprintAction)))
+	{
+		return *Result;
+	}
+
+	return false;
+}
+
+TSharedPtr<TMap<FActionFilterCacheKey, bool>> FActionFilterCache::GetCachedFilterResults(const FBlueprintActionFilter& Filter)
+{
+	// if the leaf for this context has already been looked up, just grab it rather than researching the tree
+	if (CachedFilterResults.IsValid())
+	{
+		return CachedFilterResults;
+	}
+
+	TSharedPtr<FActionFilterCacheNode> Current = SharedCache;
+
+	/** update Current to it's child that matches Key. Create one if neccesary */
+#define LOOKUP_CACHE_KEY(Key) \
+	FindOrCacheChildNode(Current, Key, #Key)
+
+/** update Current to it's child that matches Key as a leaf. Create one if neccesary */
+#define LOOKUP_CACHE_KEY_LEAF(Key) \
+	FindOrCacheLeaf(Current, Key, #Key)
+
+/** update Current to it's child that matches Key as a leaf. Create one if neccesary
+ *  uses Comparison lambda to check for key "equality"
+ */
+#define LOOKUP_CACHE_KEY_PREDICATE(Key, Comparison) \
+	FindOrCacheChildNode<decltype(Key)>(Current, Key, #Key, TDelegate<bool(const decltype(Key)&, const decltype(Key)&)>::CreateLambda(Comparison))
+
+ /** update Current to it's child that matches Key as a leaf. Create one if neccesary
+  *  uses Comparison lambda to check for key "equality"
+  */
+#define LOOKUP_CACHE_KEY_LEAF_PREDICATE(Key, Comparison) \
+	FindOrCacheLeaf<decltype(Key)>(Current, Key, #Key, TDelegate<bool(const decltype(Key)&, const decltype(Key)&)>::CreateLambda(Comparison))
+
+  // traverse/build the cache tree looking for the cached results
+
+  // cache based on what filter tests are being applied
+  // @todo - this is no longer a struct type, so we would need some way to infer delegate equality (or maybe we can generate a filter GUID instead of the stuff below)
+//	LOOKUP_CACHE_KEY(Filter.GetFilterTests());
+
+	// cache based on which blueprints/graphs/editor we're generating actions for
+	LOOKUP_CACHE_KEY(Filter.Context.Blueprints);
+	LOOKUP_CACHE_KEY(Filter.Context.Graphs);
+	LOOKUP_CACHE_KEY(Filter.Context.EditorPtr);
+
+	// cache based on the filter's target classes
+	LOOKUP_CACHE_KEY(Filter.TargetClasses);
+
+	// cache based on what node types are permitted or rejected
+	LOOKUP_CACHE_KEY(Filter.RejectedNodeTypes);
+	LOOKUP_CACHE_KEY(Filter.PermittedNodeTypes);
+
+	// cache based on what flags are active
+//	LOOKUP_CACHE_KEY(Filter.GetFilterFlags());
+
+	// cache based on the relevant pin types and directions
+	LOOKUP_CACHE_KEY_LEAF_PREDICATE(Filter.Context.Pins, [](const TArray<UEdGraphPin*>& A, const TArray<UEdGraphPin*>& B)
+	{
+		if (A.Num() != B.Num())
+		{
+			return false;
+		}
+		for (int32 Index = 0; Index < A.Num(); ++Index)
+		{
+			if (A[Index]->Direction != B[Index]->Direction)
+			{
+				return false;
+			}
+			if (A[Index]->PinType != B[Index]->PinType)
+			{
+				return false;
+			}
+		}
+		return true;
+	});
+
+	// at this point, IsFilteredByThisCache has been updated to a shared memory space that can be retrieved again
+	// by following the above procedure with the same keys values.
+	return CachedFilterResults;
+}
+
+void FActionFilterCache::PopCache()
+{
+	int32 MinIndex = 0;
+	uint32 MinAccessCount = TNumericLimits<uint32>::Max();
+	for (int32 LeafIndex = 0; LeafIndex < CacheLeafs.Num(); ++LeafIndex)
+	{
+		TWeakPtr<FActionFilterCacheNode>& Leaf = CacheLeafs[LeafIndex];
+		// if there's invalid cache, remove it first
+		if (!Leaf.IsValid())
+		{
+			MinIndex = LeafIndex;
+			break;
+		}
+
+		const uint32 AccessCount = Leaf.Pin()->GetLeafAccessCount();
+		if (AccessCount < MinAccessCount)
+		{
+			MinAccessCount = AccessCount;
+			MinIndex = LeafIndex;
+		}
+
+		// decrement access count of all leafs so that recently used leafs get prioritized more
+		Leaf.Pin()->DecrementLeafAccessCount();
+	}
+
+	TWeakPtr<FActionFilterCacheNode> RemovedElement = CacheLeafs[MinIndex];
+	if (RemovedElement.IsValid())
+	{
+		RemovedElement.Pin()->UnCacheSelf();
+	}
+
+	CacheLeafs.RemoveAt(MinIndex, 1, false);
+}
+
+TSharedPtr<FActionFilterCacheNode> FActionFilterCache::SharedCache = MakeShared<FActionFilterCacheNode>();
+
+TAutoConsoleVariable<int32> FActionFilterCache::CVarCacheLeafCapacity{
+	TEXT("BP.ActionMenuFilterCacheLeafCapacity"),
+	32,
+	TEXT("The number of action menu contexts to cache simultaniously. raising this number will increase the memory footprint but decrease how often the cache is blown")
+};
+
+TArray<TWeakPtr<FActionFilterCacheNode>> FActionFilterCache::CacheLeafs;
+
+FDelegateHandle FActionFilterCache::ReloadCompleteDelegate;
+FDelegateHandle FActionFilterCache::ModuleChangedDelegate;
\ No newline at end of file
diff --git a/Engine/Source/Editor/BlueprintGraph/Private/BlueprintActionFilterResultCache.h b/Engine/Source/Editor/BlueprintGraph/Private/BlueprintActionFilterResultCache.h
new file mode 100644
--- /dev/null
+++ b/Engine/Source/Editor/BlueprintGraph/Private/BlueprintActionFilterResultCache.h
@@ -0,0 +1,238 @@
+// Copyright Epic Games, Inc. All Rights Reserved.
+
+#pragma once
+
+#include "BlueprintActionFilter.h"
+#include "HAL/IConsoleManager.h"
+
+template<typename TKeyType>
+class TActionFilterCacheKeyNode;
+template<typename TKeyType>
+class TActionFilterCacheLeaf;
+
+/*******************************************************************************
+ * FActionFilterCacheKey
+ ******************************************************************************/
+
+// Wrapper for a hashed action key
+struct FActionFilterCacheKey
+{
+	uint32 KeyHash;
+
+	explicit FActionFilterCacheKey(const FBlueprintActionInfo& InAction);
+};
+
+bool operator==(const FActionFilterCacheKey& A, const FActionFilterCacheKey& B);
+uint32 GetTypeHash(const FActionFilterCacheKey& InKey);
+
+/*******************************************************************************
+ * FActionFilterCacheNode
+ ******************************************************************************/
+
+ /**
+  * This class is used by FActionFilterCache to create a tree structure that is used to store cached filter results in
+  * such a way that they can be looked up using all the relevant context data that determines how items are filtered.
+  * each node has a type agnostic key so that we can easily change how the cache is keyed in the future.
+  * (See FActionFilterCache::GetIsFilteredByThisCache)
+  */
+class FActionFilterCacheNode : public TSharedFromThis<FActionFilterCacheNode>
+{
+public:
+	template<typename TKeyType>
+	friend class TActionFilterCacheKeyNode;
+
+	// Branching nodes
+	using FWeakNode = TWeakPtr<FActionFilterCacheNode>;
+	template<typename TKeyType>
+	using TWeakKeyNode = TWeakPtr<TActionFilterCacheKeyNode<TKeyType>>;
+	template<typename TKeyType>
+	using TWeakLeaf = TWeakPtr<TActionFilterCacheLeaf<TKeyType>>;
+
+	// Leaf nodes
+	using FSharedNode = TSharedPtr<FActionFilterCacheNode>;
+	template<typename TKeyType>
+	using TSharedKeyNode = TSharedPtr<TActionFilterCacheKeyNode<TKeyType>>;
+	template<typename TKeyType>
+	using TSharedLeaf = TSharedPtr<TActionFilterCacheLeaf<TKeyType>>;
+
+	FActionFilterCacheNode() = default;
+	virtual ~FActionFilterCacheNode() = default;
+
+	/** Looks up a child of this node by key. Note: the key is not hashed. This is O(n) */
+	template<typename TKeyType>
+	TSharedKeyNode<TKeyType> Find(const TKeyType& Key) const
+	{
+		for (const TSharedPtr<FActionFilterCacheNode>& Child : Children)
+		{
+			if (Child->KeysMatch(&Key))
+			{
+				return StaticCastSharedPtr<TActionFilterCacheKeyNode<TKeyType>>(Child);
+			}
+		}
+		return nullptr;
+	}
+
+	/** Looks up a child of this node by key and creates one if one doesnt exist.
+	 *	The node that is created will be a TActionFilterCacheKeyNode. Use FindOrCacheChildLeaf if you need it to be a
+	 *	TActionFilterCacheLeaf
+	 *  Note: the key is not hashed. This is O(n) */
+	template<typename TKeyType>
+	TSharedKeyNode<TKeyType> FindOrCacheChildNode(const TKeyType& Key, const FString& InDebugName = FString(), TDelegate<bool(const TKeyType&, const TKeyType&)> Comparison = nullptr, bool* OutFoundWithoutCaching = nullptr)
+	{
+		TSharedKeyNode<TKeyType> Result = Find(Key);
+		if (OutFoundWithoutCaching)
+		{
+			*OutFoundWithoutCaching = Result.IsValid();
+		}
+
+		if (Result.IsValid())
+		{
+			return Result;
+		}
+
+#if UE_BUILD_DEBUG
+		DebugName = InDebugName;
+#endif
+		return Cache(new TActionFilterCacheKeyNode<TKeyType>(Key, Comparison));
+	}
+
+	/** Looks up a child of this node by key and creates one if one doesnt exist.
+	 *	The node that is created will be a TActionFilterCacheLeaf. Use FindOrCacheChildNode if you need it to be a
+	 *	TActionFilterCacheKeyNode
+	 *  Note: the key is not hashed. This is O(n) */
+	template<typename TKeyType>
+	TSharedLeaf<TKeyType> FindOrCacheChildLeaf(const TKeyType& Key, const FString& InDebugName = FString(), TDelegate<bool(const TKeyType&, const TKeyType&)> Comparison = nullptr, bool* OutFoundWithoutCaching = nullptr)
+	{
+		TSharedKeyNode<TKeyType> Result = Find(Key);
+		if (OutFoundWithoutCaching)
+		{
+			*OutFoundWithoutCaching = Result.IsValid();
+		}
+
+		if (Result.IsValid())
+		{
+			TSharedLeaf<TKeyType> Leaf = StaticCastSharedPtr<TActionFilterCacheLeaf<TKeyType>>(Result);
+			++Leaf->AccessCount;
+			return Leaf;
+		}
+#if UE_BUILD_DEBUG
+		DebugName = InDebugName;
+#endif
+		return Cache(new TActionFilterCacheLeaf<TKeyType>(Key, Comparison));
+	}
+
+	/** get the number of times this leaf has been accessed
+	 *  Warning: throws and error if this isn't a TActionFilterCacheLeaf */
+	virtual uint32 GetLeafAccessCount() const
+	{
+		checkf(false, TEXT("GetLeafAccessCount was called on a non-leaf node"))
+			return 0;
+	}
+
+	/** decrement the number of times this leaf has been accessed
+	 *  Warning: throws and error if this isn't a TActionFilterCacheLeaf */
+	virtual void DecrementLeafAccessCount()
+	{
+		checkf(false, TEXT("DecrementLeafAccessCount was called on a non-leaf node"))
+	}
+
+	/** remove this node from the the tree */
+	virtual void UnCacheSelf() {}
+
+#if UE_BUILD_DEBUG
+	FString DebugName;
+#endif
+
+protected:
+	virtual bool KeysMatch(const void* RawKey) const
+	{
+		checkf(false, TEXT("Tried to call KeysMatch on a Node with no keys"))
+			return false;
+	}
+
+	template<typename TNodeType>
+	TSharedPtr<TNodeType> Cache(TNodeType* Child)
+	{
+		Child->Parent = AsShared();
+		Children.Emplace(MakeShareable<FActionFilterCacheNode>(Child));
+		return StaticCastSharedPtr<TNodeType>(Children.Last());
+	}
+
+	TArray<FSharedNode> Children;
+	FWeakNode Parent;
+};
+
+/*******************************************************************************
+ * FActionFilterCache
+ ******************************************************************************/
+
+ /**
+  * This class creates a tree structure that is used to store cached filter results in such a way that they can be
+  * looked up using all the relevant context data that determines how items are filtered.
+  * each node has a type agnostic key so that we can easily change how the cache is keyed.
+  * (See FActionFilterCache::GetIsFilteredByThisCache)
+  */
+class FActionFilterCache
+{
+public:
+	FActionFilterCache();
+
+	static void OnReloadComplete(EReloadCompleteReason Reason);
+	static void OnModuleChanged(FName ModuleName, EModuleChangeReason ChangeReason);
+	static void ClearAllCache();
+
+private:
+	/** get the cached map used by FBlueprintActionFilter::IsFilteredByThis to determine whether BP actions are filtered out */
+	TSharedPtr<TMap<FActionFilterCacheKey, bool>> GetCachedFilterResults(const FBlueprintActionFilter& Filter);
+
+	template<typename TKeyType>
+	void FindOrCacheChildNode(TSharedPtr<FActionFilterCacheNode>& Current, const TKeyType& Key, FString DebugName, TDelegate<bool(const TKeyType&, const TKeyType&)> Comparison = nullptr)
+	{
+		TSharedPtr<TActionFilterCacheKeyNode<TKeyType>> Found = Current->FindOrCacheChildNode(Key, DebugName, Comparison);
+		Current = StaticCastSharedPtr<FActionFilterCacheNode>(Found);
+
+	}
+
+	template<typename TKeyType>
+	void FindOrCacheLeaf(TSharedPtr<FActionFilterCacheNode>& Current, const TKeyType& Key, FString DebugName, TDelegate<bool(const TKeyType&, const TKeyType&)> Comparison = nullptr)
+	{
+		bool bFoundWithoutCaching = false;
+		TSharedPtr<TActionFilterCacheLeaf<TKeyType>> Found = Current->FindOrCacheChildLeaf(Key, DebugName, Comparison, &bFoundWithoutCaching);
+		CachedFilterResults = Found->CachedFilterResults;
+		Current = StaticCastSharedPtr<FActionFilterCacheNode>(Found);
+
+		// if we just expanded the cache, log it and check if the cache is getting too big
+		if (!bFoundWithoutCaching)
+		{
+			if (CacheLeafs.Num() >= CVarCacheLeafCapacity.GetValueOnGameThread())
+			{
+				PopCache();
+			}
+			CacheLeafs.Push(Found);
+		}
+	}
+
+	/** Find the least likely item in the cache to be used again, and remove it to make room for more data */
+	static void PopCache();
+
+	/** Cache result */
+	void CacheFilteredResult(const FBlueprintActionFilter& Filter, FBlueprintActionInfo& BlueprintAction, bool bWasFilteredOut);
+
+	/** Filter function */
+	bool IsFilteredByCachedResult(const FBlueprintActionFilter& Filter, FBlueprintActionInfo& BlueprintAction);
+
+private:
+	/** the shared tree that stores previously calculated filter results for several different contects */
+	static TSharedPtr<FActionFilterCacheNode> SharedCache;
+
+	static TAutoConsoleVariable<int32> CVarCacheLeafCapacity;
+
+	/** array of all the leaf nodes in SharedCache so we can pop a rarely used node if we exceed CacheLeafCapacity */
+	static TArray<TWeakPtr<FActionFilterCacheNode>> CacheLeafs;
+
+	static FDelegateHandle ReloadCompleteDelegate;
+	static FDelegateHandle ModuleChangedDelegate;
+
+	/** cached results from GetIsFilteredByThisCache so we don't have to re-search the tree */
+	TSharedPtr<TMap<FActionFilterCacheKey, bool>> CachedFilterResults;
+};
\ No newline at end of file
diff --git a/Engine/Source/Editor/BlueprintGraph/Public/BlueprintActionFilter.h b/Engine/Source/Editor/BlueprintGraph/Public/BlueprintActionFilter.h
--- a/Engine/Source/Editor/BlueprintGraph/Public/BlueprintActionFilter.h
+++ b/Engine/Source/Editor/BlueprintGraph/Public/BlueprintActionFilter.h
@@ -26,28 +26,20 @@
 #include "UObject/WeakObjectPtr.h"
 #include "UObject/WeakObjectPtrTemplates.h"
 
-class FBlueprintActionFilter;
+#define ENABLE_BLUEPRINT_ACTION_FILTER_PROFILING (!(UE_BUILD_SHIPPING || UE_BUILD_TEST) && 0)
+
 class FProperty;
 class UClass;
-class UEdGraphNode;
 class UFunction;
 class UObject;
-
-#define ENABLE_BLUEPRINT_ACTION_FILTER_PROFILING (!(UE_BUILD_SHIPPING || UE_BUILD_TEST) && 0)
-
-
-
-class FActionFilterCacheNode;
 class IAssetReferenceFilter;
 class IBlueprintEditor;
 class UBlueprint;
 class UBlueprintNodeSpawner;
 class UEdGraph;
 class UEdGraphPin;
-template<typename TKeyType>
-class TActionFilterCacheKeyNode;
-template<typename TKeyType>
-class TActionFilterCacheLeaf;
+class UEdGraphNode;
+
 /*******************************************************************************
  * FBlueprintActionContext
  ******************************************************************************/
@@ -168,8 +160,6 @@ struct BLUEPRINTGRAPH_API FBlueprintActionInfo
 	 */
 	UFunction const* GetAssociatedFunction();
 
-	bool operator==(const FBlueprintActionInfo& other) const;
-
 	/** The raw action that this struct represent (const so we don't mutate the database) */
 	UBlueprintNodeSpawner const* const NodeSpawner;
 
@@ -190,231 +180,24 @@ private:
 	IBlueprintNodeBinder::FBindingSet Bindings;
 };
 
-uint32 GetTypeHash(const FBlueprintActionInfo& BlueprintActionInfo);
-
-/*******************************************************************************
- * FActionFilterCacheNode
- ******************************************************************************/
-
-/**
- * This class is used by FActionFilterCache to create a tree structure that is used to store cached filter results in
- * such a way that they can be looked up using all the relevant context data that determines how items are filtered.
- * each node has a type agnostic key so that we can easily change how the cache is keyed in the future.
- * (See FActionFilterCache::GetIsFilteredByThisCache)
- */
-class FActionFilterCacheNode : public TSharedFromThis<FActionFilterCacheNode>
-{
-public:
-	template<typename TKeyType>
-	friend class TActionFilterCacheKeyNode;
-
-	// Branching nodes
-	using FWeakNode = TWeakPtr<FActionFilterCacheNode>;
-	template<typename TKeyType>
-	using TWeakKeyNode = TWeakPtr<TActionFilterCacheKeyNode<TKeyType>>;
-	template<typename TKeyType>
-	using TWeakLeaf = TWeakPtr<TActionFilterCacheLeaf<TKeyType>>;
-
-	// Leaf nodes
-	using FSharedNode = TSharedPtr<FActionFilterCacheNode>;
-	template<typename TKeyType>
-	using TSharedKeyNode = TSharedPtr<TActionFilterCacheKeyNode<TKeyType>>;
-	template<typename TKeyType>
-	using TSharedLeaf = TSharedPtr<TActionFilterCacheLeaf<TKeyType>>;
-	
-	FActionFilterCacheNode() = default;
-	virtual ~FActionFilterCacheNode() = default;
-
-	/** Looks up a child of this node by key. Note: the key is not hashed. This is O(n) */
-	template<typename TKeyType>
-	TSharedKeyNode<TKeyType> Find(const TKeyType& Key) const
-	{
-		for (const TSharedPtr<FActionFilterCacheNode>& Child : Children)
-		{
-			if (Child->KeysMatch(&Key))
-			{
-				return StaticCastSharedPtr<TActionFilterCacheKeyNode<TKeyType>>(Child);
-			}
-		}
-		return nullptr;
-	}
-
-	/** Looks up a child of this node by key and creates one if one doesnt exist.
-	 *	The node that is created will be a TActionFilterCacheKeyNode. Use FindOrCacheChildLeaf if you need it to be a
-	 *	TActionFilterCacheLeaf
-	 *  Note: the key is not hashed. This is O(n) */
-	template<typename TKeyType>
-	TSharedKeyNode<TKeyType> FindOrCacheChildNode(const TKeyType& Key, const FString& InDebugName = FString(), TDelegate<bool(const TKeyType&, const TKeyType&)> Comparison = nullptr, bool *OutFoundWithoutCaching = nullptr)
-	{
-		TSharedKeyNode<TKeyType> Result = Find(Key);
-		if (OutFoundWithoutCaching)
-		{
-			*OutFoundWithoutCaching = Result.IsValid();
-		}
-		
-		if (Result.IsValid())
-		{
-			return Result;
-		}
-		
-#if UE_BUILD_DEBUG
-		DebugName = InDebugName;
-#endif
-		return Cache(new TActionFilterCacheKeyNode<TKeyType>(Key, Comparison));
-	}
-
-	/** Looks up a child of this node by key and creates one if one doesnt exist.
-	 *	The node that is created will be a TActionFilterCacheLeaf. Use FindOrCacheChildNode if you need it to be a
-	 *	TActionFilterCacheKeyNode
-	 *  Note: the key is not hashed. This is O(n) */
-	template<typename TKeyType>
-	TSharedLeaf<TKeyType> FindOrCacheChildLeaf(const TKeyType& Key, const FString& InDebugName = FString(), TDelegate<bool(const TKeyType&, const TKeyType&)> Comparison = nullptr, bool *OutFoundWithoutCaching = nullptr)
-	{
-		TSharedKeyNode<TKeyType> Result = Find(Key);
-		if (OutFoundWithoutCaching)
-		{
-			*OutFoundWithoutCaching = Result.IsValid();
-		}
-		
-		if (Result.IsValid())
-		{
-			TSharedLeaf<TKeyType> Leaf = StaticCastSharedPtr<TActionFilterCacheLeaf<TKeyType>>(Result);
-			++Leaf->AccessCount;
-			return Leaf;
-		}
-#if UE_BUILD_DEBUG
-		DebugName = InDebugName;
-#endif
-		return Cache(new TActionFilterCacheLeaf<TKeyType>(Key, Comparison));
-	}
-
-	/** get the number of times this leaf has been accessed
-	 *  Warning: throws and error if this isn't a TActionFilterCacheLeaf */
-	virtual uint32 GetLeafAccessCount() const
-	{
-		checkf(false, TEXT("GetLeafAccessCount was called on a non-leaf node"))
-		return 0;
-	}
-
-	/** decrement the number of times this leaf has been accessed
-	 *  Warning: throws and error if this isn't a TActionFilterCacheLeaf */
-	virtual void DecrementLeafAccessCount()
-	{
-		checkf(false, TEXT("DecrementLeafAccessCount was called on a non-leaf node"))
-	}
-
-	/** remove this node from the the tree */
-	virtual void UnCacheSelf(){}
-
-#if UE_BUILD_DEBUG
-	FString DebugName;
-#endif
-
-protected:
-	virtual bool KeysMatch(const void* RawKey) const
-	{
-		checkf(false, TEXT("Tried to call KeysMatch on a Node with no keys"))
-		return false;
-	}
-
-	template<typename TNodeType>
-	TSharedPtr<TNodeType> Cache(TNodeType* Child)
-	{
-		Child->Parent = AsShared();
-		Children.Emplace(MakeShareable<FActionFilterCacheNode>(Child));
-		return StaticCastSharedPtr<TNodeType>(Children.Last());
-	}
-	
-	TArray<FSharedNode> Children;
-	FWeakNode Parent;
-};
-
-/*******************************************************************************
- * FActionFilterCache
- ******************************************************************************/
-
-/**
- * This class creates a tree structure that is used to store cached filter results in such a way that they can be
- * looked up using all the relevant context data that determines how items are filtered.
- * each node has a type agnostic key so that we can easily change how the cache is keyed.
- * (See FActionFilterCache::GetIsFilteredByThisCache)
- */
-class FActionFilterCache
-{
-public:
-	FActionFilterCache();
-
-	static void OnReloadComplete(EReloadCompleteReason Reason);
-	static void OnModuleChanged(FName ModuleName, EModuleChangeReason ChangeReason);
-	static void ClearAllCache();
-
-	/** get the cached map used by FBlueprintActionFilter::IsFilteredByThis to determine whether BP actions are filtered out */
-	TSharedPtr<TMap<FBlueprintActionInfo, bool>> GetCachedFilterResults(const FBlueprintActionFilter &Filter);
-
-private:
-	template<typename TKeyType>
-	void FindOrCacheChildNode(TSharedPtr<FActionFilterCacheNode> &Current, const TKeyType& Key, FString DebugName, TDelegate<bool(const TKeyType&, const TKeyType&)> Comparison = nullptr)
-	{
-		TSharedPtr<TActionFilterCacheKeyNode<TKeyType>> Found = Current->FindOrCacheChildNode(Key, DebugName, Comparison);
-		Current = StaticCastSharedPtr<FActionFilterCacheNode>(Found);
-		
-	}
-
-	template<typename TKeyType>
-	void FindOrCacheLeaf(TSharedPtr<FActionFilterCacheNode> &Current, const TKeyType& Key, FString DebugName, TDelegate<bool(const TKeyType&, const TKeyType&)> Comparison = nullptr)
-	{
-		bool bFoundWithoutCaching = false;
-		TSharedPtr<TActionFilterCacheLeaf<TKeyType>> Found = Current->FindOrCacheChildLeaf(Key, DebugName, Comparison, &bFoundWithoutCaching);
-		CachedFilterResults = Found->CachedFilterResults;
-		Current = StaticCastSharedPtr<FActionFilterCacheNode>(Found);
-
-		// if we just expanded the cache, log it and check if the cache is getting too big
-		if (!bFoundWithoutCaching)
-		{
-			if (CacheLeafs.Num() >= CVarCacheLeafCapacity.GetValueOnGameThread())
-			{
-				PopCache();
-			}
-			CacheLeafs.Push(Found);
-		}
-	}
-
-	/** Find the least likely item in the cache to be used again, and remove it to make room for more data */
-	static void PopCache();
-
-private:
-	/** the shared tree that stores previously calculated filter results for several different contects */
-	static TSharedPtr<FActionFilterCacheNode> SharedCache;
-	
-	static TAutoConsoleVariable<int32> CVarCacheLeafCapacity;
-
-	/** array of all the leaf nodes in SharedCache so we can pop a rarely used node if we exceed CacheLeafCapacity */
-	static TArray<TWeakPtr<FActionFilterCacheNode>> CacheLeafs;
-
-	static FDelegateHandle ReloadCompleteDelegate;
-	static FDelegateHandle ModuleChangedDelegate;
-
-	/** cached results from GetIsFilteredByThisCache so we don't have to re-search the tree */
-	TSharedPtr<TMap<FBlueprintActionInfo, bool>> CachedFilterResults;
-};
-
 /*******************************************************************************
  * EBlueprintActionFilterRejectionTestFlags
  ******************************************************************************/
 
+// @todo - Need to revisit this later; for now these flags remain in place for backcompat but are currently not in use.
 enum class EActionFilterTestFlags
 {
 	None = 0,
-	
+
 	/**
 	 * these tests will be have their combined result cached for faster lookup.
 	 * note: rejection tests can only be cached if they either
 	 *		A: Don't access any mutable data
 	 *		or B: Only access mutable data that is listed as a filter key in FActionFilterCache::GetCachedFilterResults
 	 */
-	CacheResults = 1 << 0,
-	
-	Default = None, // this would allow us to change the default behavior
+	 CacheResults = 1 << 0,
+
+	 Default = None, // this would allow us to change the default behavior
 };
 
 ENUM_CLASS_FLAGS(EActionFilterTestFlags);
@@ -423,10 +206,11 @@ ENUM_CLASS_FLAGS(EActionFilterTestFlags);
  * FActionFilterTest
 ******************************************************************************/
 
+// @todo - Need to revisit this later; for now this wrapper type remains in place for backcompat, but flags are currently not in use.
 struct BLUEPRINTGRAPH_API FActionFilterTest : TSharedFromThis<FActionFilterTest>
 {
 	static TAutoConsoleVariable<bool> CVarEnableCaching;
-	
+
 	/** The filter uses a series of rejection tests matching */
 	typedef FBlueprintGraphModule::FActionMenuRejectionTest FRejectionDelegate;
 
@@ -437,24 +221,16 @@ struct BLUEPRINTGRAPH_API FActionFilterTest : TSharedFromThis<FActionFilterTest>
 		EnumAddFlags(Flags, InFlags);
 		return AsShared();
 	}
-	
-	bool Call(FBlueprintActionFilter const& Filter, FBlueprintActionInfo& ActionInfo, TMap<FString, uint32> *FilterTestToCpuSpecIds = nullptr);
+
+	bool Call(FBlueprintActionFilter const& Filter, FBlueprintActionInfo& ActionInfo);
 	bool operator==(const FActionFilterTest& other) const;
-	
+
 	FRejectionDelegate RejectionDelegate;
 	FString Name; // this name must be unique
 	EActionFilterTestFlags Flags = EActionFilterTestFlags::Default;
-	
-#if ENABLE_BLUEPRINT_ACTION_FILTER_PROFILING
-	struct FFilterTestProfileRecord
-	{
-		float TotalTimeMs = 0.0f;
-		int32 NumIterations = 0;
-		int32 NumFilteredOut = 0;
-	} ProfileRecord;
-#endif	// ENABLE_BLUEPRINT_ACTION_FILTER_PROFILING
 };
 
+// @todo - Need to revisit this; currently defined for backcompat but we may deprecate its use later as it is somewhat rigid.
 // constructor for TSharedRef<FBlueprintActionFilterRejectionTest> that auto deduces the name
 #define MAKE_ACTION_FILTER_REJECTION_TEST(TestFunc, ...) \
 	MakeShared<FActionFilterTest>(FActionFilterTest::FRejectionDelegate::CreateStatic(TestFunc, ##__VA_ARGS__), TEXT(#TestFunc))
@@ -466,9 +242,9 @@ struct BLUEPRINTGRAPH_API FActionFilterTest : TSharedFromThis<FActionFilterTest>
 class BLUEPRINTGRAPH_API FBlueprintActionFilter
 {
 public:
-	friend FActionFilterCache;
+	/** The filter uses a series of rejection tests matching */
+	typedef FBlueprintGraphModule::FActionMenuRejectionTest FRejectionTestDelegate;
 
-public:
 	enum EFlags // Flags, which configure certain rejection tests.
 	{
 		BPFILTER_NoFlags					= (0),
@@ -572,7 +348,12 @@ public:
 	 * Users can extend the filter and add their own rejection tests with this
 	 * method. We use rejection "IsFiltered" tests rather than inclusive tests 
 	 * because it is more optimal to whittle down the list of actions early.
-	 * 
+	 *
+	 * @param  RejectionTestDelegate	The rejection test you wish to add to this filter.
+	 */
+	void AddRejectionTest(FRejectionTestDelegate RejectionTestDelegate);
+
+	/**
 	 * @param  RejectionTest	a test this filter will run to cull blueprint actions
 	 * @param  Flags	List of flags that will be added in the provided rejection test
 	 */
@@ -612,18 +393,10 @@ public:
 	 */
 	FBlueprintActionFilter const& operator&=(FBlueprintActionFilter const& Rhs);
 
-	TSharedPtr<TMap<FBlueprintActionInfo, bool>> GetCachedFilterResults();
-
-	const TArray<TSharedRef<FActionFilterTest>> &GetFilterTests() const;
-	
-#if ENABLE_BLUEPRINT_ACTION_FILTER_PROFILING
-	TArray<FString> GetFilterTestProfile();
-#endif	// ENABLE_BLUEPRINT_ACTION_FILTER_PROFILING
-
 	/**
 	 * @return TRUE if any of the given configuration flags are set on this filter.
 	 */
-	FORCEINLINE bool HasAnyFlags(EFlags InFlags) const
+	inline bool HasAnyFlags(EFlags InFlags) const
 	{
 		return (FilterFlags & InFlags) != 0;
 	}
@@ -631,11 +404,17 @@ public:
 	/**
 	 * @return TRUE if all of the given configuration flags are set on this filter.
 	 */
-	FORCEINLINE bool HasAllFlags(EFlags InFlags) const
+	inline bool HasAllFlags(EFlags InFlags) const
 	{
 		return (FilterFlags & InFlags) == InFlags;
 	}
 
+#if ENABLE_BLUEPRINT_ACTION_FILTER_PROFILING
+	static bool IsFilterTestTraceLoggingEnabled();
+	static bool IsFilterTestStatsLoggingEnabled();
+	TArray<FString> GetFilterTestProfile();
+#endif	// ENABLE_BLUEPRINT_ACTION_FILTER_PROFILING
+
 private:
 	/**
 	 * Query to check and see if the specified action gets filtered out by this 
@@ -644,22 +423,13 @@ private:
 	 * @param  BlueprintAction	The node-spawner you wish to test.
 	 * @return False if the action passes the filter, otherwise true (the action got filtered out).
 	 */
-	bool IsFilteredByThis(FBlueprintActionInfo& BlueprintAction);
-
-	/**
-	 * Query to check and see if the specified action gets filtered out by the filter tests
-	 * that have been marked as cacheable
-	 * 
-	 * @param  BlueprintAction	The node-spawner you wish to test.
-	 * @return False if the action passes the filter tests, otherwise true (the action got filtered out).
-	 */
-	bool IsFilteredByCachedFilters(FBlueprintActionInfo& BlueprintAction);
+	bool IsFilteredByThis(FBlueprintActionInfo& BlueprintAction) const;
 
 	/** Configuration flags for this filter. */
 	EFlags FilterFlags = BPFILTER_NoFlags;
 
 	/** Set of rejection tests for this specific filter. */
-	TArray<TSharedRef<FActionFilterTest>> FilterTests;
+	TArray<FRejectionTestDelegate> FilterTests;
 
 	/** Filters to be logically and'd in with the IsFilteredByThis() result. */
 	TArray<FBlueprintActionFilter> AndFilters;
@@ -667,7 +437,20 @@ private:
 	/** Alternative filters to be logically or'd in with the IsFilteredByThis() result. */
 	TArray<FBlueprintActionFilter> OrFilters;
 
-	FActionFilterCache ActionFilterCache;
+#if ENABLE_BLUEPRINT_ACTION_FILTER_PROFILING
+	struct FFilterTestProfileRecord
+	{
+		uint64 TestFuncPtr = 0;
+		float TotalTimeMs = 0.0f;
+		int32 NumIterations = 0;
+		int32 NumFilteredOut = 0;
+	};
+
+	mutable TMap<int32, FFilterTestProfileRecord> FilterTestProfiles;
+
+	FFilterTestProfileRecord* BeginFilterTestProfileEvent(int32 FilterTestIndex, uint64 FilterTestFuncPtr) const;
+	void EndFilterTestProfileEvent(FFilterTestProfileRecord* InEvent, bool bWasFiltered, double StartTime) const;
+#endif	// ENABLE_BLUEPRINT_ACTION_FILTER_PROFILING
 };
 
 ENUM_CLASS_FLAGS(FBlueprintActionFilter::EFlags);
\ No newline at end of file
diff --git a/Engine/Source/Editor/Kismet/Private/BlueprintActionMenuBuilder.cpp b/Engine/Source/Editor/Kismet/Private/BlueprintActionMenuBuilder.cpp
--- a/Engine/Source/Editor/Kismet/Private/BlueprintActionMenuBuilder.cpp
+++ b/Engine/Source/Editor/Kismet/Private/BlueprintActionMenuBuilder.cpp
@@ -16,17 +16,13 @@
 #include "EditorCategoryUtils.h"
 #include "ObjectEditorUtils.h"
 
-#define LOCTEXT_NAMESPACE "BlueprintActionMenuBuilder"
-DEFINE_LOG_CATEGORY_STATIC(LogBlueprintActionMenuItemFactory, Log, All);
-
 #if ENABLE_BLUEPRINT_ACTION_FILTER_PROFILING
-static TAutoConsoleVariable<bool> CVarBPEnableActionMenuFilterTestStatsLogging(
-	TEXT("BP.EnableActionMenuFilterTestStatsLogging"),
-	false,
-	TEXT("If enabled, filter test stats will be logged for each menu section each time a Blueprint action menu is built for display.")
-);
+#include "Misc/OutputDeviceFile.h"
 #endif	// ENABLE_BLUEPRINT_ACTION_FILTER_PROFILING
 
+#define LOCTEXT_NAMESPACE "BlueprintActionMenuBuilder"
+DEFINE_LOG_CATEGORY_STATIC(LogBlueprintActionMenuItemFactory, Log, All);
+
 /*******************************************************************************
  * FBlueprintActionMenuItemFactory
  ******************************************************************************/
@@ -53,11 +49,10 @@ public:
 	 * Spawns a new FBlueprintActionMenuItem with the node-spawner. Constructs
 	 * the menu item's category, name, tooltip, etc.
 	 * 
-	 * @param  EditorContext	
 	 * @param  Action			The node-spawner that the new menu item should wrap.
 	 * @return A newly allocated FBlueprintActionMenuItem (which wraps the supplied action).
 	 */
-	TSharedPtr<FBlueprintActionMenuItem> MakeActionMenuItem(TWeakPtr<FBlueprintEditor> EditorContext, FBlueprintActionInfo const& ActionInfo);
+	TSharedPtr<FBlueprintActionMenuItem> MakeActionMenuItem(FBlueprintActionInfo const& ActionInfo);
 
 	/**
 	 * Spawns a new FBlueprintDragDropMenuItem with the node-spawner. Constructs
@@ -74,7 +69,7 @@ public:
 	 * @param  BoundAction	
 	 * @return 
 	 */
-	TSharedPtr<FBlueprintActionMenuItem> MakeBoundMenuItem(TWeakPtr<FBlueprintEditor> EditorContext, FBlueprintActionInfo const& ActionInfo);
+	TSharedPtr<FBlueprintActionMenuItem> MakeBoundMenuItem(FBlueprintActionInfo const& ActionInfo);
 	
 private:
 	/**
@@ -86,21 +81,15 @@ private:
 	UBlueprint* GetTargetBlueprint() const;
 
 	/**
-	 * 
-	 * 
-	 * @param  EditorContext	
 	 * @return 
 	 */
-	UEdGraph* GetTargetGraph(TWeakPtr<FBlueprintEditor> EditorContext) const;
+	UEdGraph* GetTargetGraph() const;
 
 	/**
-	 *
-	 *
-	 * @param  EditorContext
 	 * @param  ActionInfo
 	 * @return
 	 */
-	FBlueprintActionUiSpec GetActionUiSignature(TWeakPtr<FBlueprintEditor> EditorContext, FBlueprintActionInfo const& ActionInfo);
+	FBlueprintActionUiSpec GetActionUiSignature(FBlueprintActionInfo const& ActionInfo);
 };
 
 //------------------------------------------------------------------------------
@@ -112,11 +101,11 @@ FBlueprintActionMenuItemFactory::FBlueprintActionMenuItemFactory(FBlueprintActio
 }
 
 //------------------------------------------------------------------------------
-TSharedPtr<FBlueprintActionMenuItem> FBlueprintActionMenuItemFactory::MakeActionMenuItem(TWeakPtr<FBlueprintEditor> EditorContext, FBlueprintActionInfo const& ActionInfo)
+TSharedPtr<FBlueprintActionMenuItem> FBlueprintActionMenuItemFactory::MakeActionMenuItem(FBlueprintActionInfo const& ActionInfo)
 {
 	TRACE_CPUPROFILER_EVENT_SCOPE(FBlueprintActionMenuItemFactory::MakeActionMenuItem);
 
-	FBlueprintActionUiSpec UiSignature = GetActionUiSignature(EditorContext, ActionInfo);
+	FBlueprintActionUiSpec UiSignature = GetActionUiSignature(ActionInfo);
 
 	UBlueprintNodeSpawner const* Action = ActionInfo.NodeSpawner;
 	FBlueprintActionMenuItem* NewMenuItem = new FBlueprintActionMenuItem(Action, UiSignature, IBlueprintNodeBinder::FBindingSet(), FText::FromString(RootCategory.ToString() + TEXT('|') + UiSignature.Category.ToString()), MenuGrouping);
@@ -193,9 +182,9 @@ TSharedPtr<FBlueprintDragDropMenuItem> FBlueprintActionMenuItemFactory::MakeDrag
 }
 
 //------------------------------------------------------------------------------
-TSharedPtr<FBlueprintActionMenuItem> FBlueprintActionMenuItemFactory::MakeBoundMenuItem(TWeakPtr<FBlueprintEditor> EditorContext, FBlueprintActionInfo const& ActionInfo)
+TSharedPtr<FBlueprintActionMenuItem> FBlueprintActionMenuItemFactory::MakeBoundMenuItem(FBlueprintActionInfo const& ActionInfo)
 {
-	FBlueprintActionUiSpec UiSignature = GetActionUiSignature(EditorContext, ActionInfo);
+	FBlueprintActionUiSpec UiSignature = GetActionUiSignature(ActionInfo);
 
 	UBlueprintNodeSpawner const* Action = ActionInfo.NodeSpawner;
 	FBlueprintActionMenuItem* NewMenuItem = new FBlueprintActionMenuItem(Action, UiSignature, ActionInfo.GetBindings(), FText::FromString(RootCategory.ToString() + TEXT('|') + UiSignature.Category.ToString()), MenuGrouping);
@@ -215,7 +204,7 @@ UBlueprint* FBlueprintActionMenuItemFactory::GetTargetBlueprint() const
 }
 
 //------------------------------------------------------------------------------
-UEdGraph* FBlueprintActionMenuItemFactory::GetTargetGraph(TWeakPtr<FBlueprintEditor> EditorContext) const
+UEdGraph* FBlueprintActionMenuItemFactory::GetTargetGraph() const
 {
 	UEdGraph* TargetGraph = nullptr;
 	if (Context.Graphs.Num() > 0)
@@ -231,22 +220,22 @@ UEdGraph* FBlueprintActionMenuItemFactory::GetTargetGraph(TWeakPtr<FBlueprintEdi
 		{
 			TargetGraph = Blueprint->UbergraphPages[0];
 		}
-		else if (EditorContext.IsValid())
+		else if (Context.EditorPtr.IsValid())
 		{
-			TargetGraph = EditorContext.Pin()->GetFocusedGraph();
+			TargetGraph = Context.EditorPtr.Pin()->GetFocusedGraph();
 		}
 	}
 	return TargetGraph;
 }
 
 //------------------------------------------------------------------------------
-FBlueprintActionUiSpec FBlueprintActionMenuItemFactory::GetActionUiSignature(TWeakPtr<FBlueprintEditor> EditorContext, FBlueprintActionInfo const& ActionInfo)
+FBlueprintActionUiSpec FBlueprintActionMenuItemFactory::GetActionUiSignature(FBlueprintActionInfo const& ActionInfo)
 {
 	TRACE_CPUPROFILER_EVENT_SCOPE(FBlueprintActionMenuItemFactory::GetActionUiSignature);
 
 	UBlueprintNodeSpawner const* Action = ActionInfo.NodeSpawner;
 
-	UEdGraph* TargetGraph = GetTargetGraph(EditorContext);
+	UEdGraph* TargetGraph = GetTargetGraph();
 	Action->PrimeDefaultUiSpec(TargetGraph);
 
 	return Action->GetUiSpec(Context, ActionInfo.GetBindings());
@@ -285,21 +274,19 @@ namespace FBlueprintActionMenuBuilderImpl
 		 * FBlueprintActionMenuItem for the specified menu (does not add the 
 		 * item to the menu-builder itself).
 		 *
-		 * @param  EditorContext	
 		 * @param  DatabaseAction	The node-spawner that the new menu item should wrap.
 		 * @return An empty TSharedPtr if the action was filtered out, otherwise a newly allocated FBlueprintActionMenuItem.
 		 */
-		MenuItemList MakeMenuItems(TWeakPtr<FBlueprintEditor> EditorContext, FBlueprintActionInfo& DatabaseAction);
+		MenuItemList MakeMenuItems(FBlueprintActionInfo& DatabaseAction);
 
 		/**
 		 * 
 		 * 
-		 * @param  EditorContext	
 		 * @param  DatabaseAction	
 		 * @param  Bindings	
 		 * @return 
 		 */
-		void AddBoundMenuItems(TWeakPtr<FBlueprintEditor> EditorContext, FBlueprintActionInfo& DatabaseAction, TArray<FFieldVariant> const& Bindings, MenuItemList& MenuItemsOut);
+		void AddBoundMenuItems(FBlueprintActionInfo& DatabaseAction, TArray<FFieldVariant> const& Bindings, MenuItemList& MenuItemsOut);
 		
 		/**
 		 * Clears out any consolidated properties that this may have been 
@@ -313,14 +300,14 @@ namespace FBlueprintActionMenuBuilderImpl
 		/** Tracks the properties that we've already consolidated and passed (when using the ConsolidatePropertyActions flag)*/
 		TMap<FProperty const*, TSharedPtr<FBlueprintDragDropMenuItem>> ConsolidatedProperties;
 	};
-	
-	/**
-	 * 
-	 * 
-	 * @param  Context	
-	 * @return 
-	 */
-	static TArray<FFieldVariant> GetBindingCandidates(FBlueprintActionContext const& Context);
+
+#if ENABLE_BLUEPRINT_ACTION_FILTER_PROFILING
+	static TAutoConsoleVariable<bool> CVarBPEnableActionMenuDumpToFile(
+		TEXT("BP.EnableActionMenuDumpToFile"),
+		false,
+		TEXT("If enabled, action menu contents will be dumped to a file after building.")
+	);
+#endif	// ENABLE_BLUEPRINT_ACTION_FILTER_PROFILING
 }
 
 //------------------------------------------------------------------------------
@@ -350,8 +337,10 @@ void FBlueprintActionMenuBuilderImpl::FMenuSectionDefinition::SetSectionSortOrde
 }
 // 
 //------------------------------------------------------------------------------
-void FBlueprintActionMenuBuilderImpl::FMenuSectionDefinition::AddBoundMenuItems(TWeakPtr<FBlueprintEditor> EditorContext, FBlueprintActionInfo& DatabaseActionInfo, TArray<FFieldVariant> const& PerspectiveBindings, MenuItemList& MenuItemsOut)
+void FBlueprintActionMenuBuilderImpl::FMenuSectionDefinition::AddBoundMenuItems(FBlueprintActionInfo& DatabaseActionInfo, TArray<FFieldVariant> const& PerspectiveBindings, MenuItemList& MenuItemsOut)
 {
+	TRACE_CPUPROFILER_EVENT_SCOPE(FMenuSectionDefinition::AddBoundMenuItems);
+
 	UBlueprintNodeSpawner const* DatabaseAction = DatabaseActionInfo.NodeSpawner;
 
 	TSharedPtr<FBlueprintActionMenuItem> LastMadeMenuItem;
@@ -392,7 +381,7 @@ void FBlueprintActionMenuBuilderImpl::FMenuSectionDefinition::AddBoundMenuItems(
 			{
 				if (!bConsolidate || !LastMadeMenuItem.IsValid())
 				{
-					LastMadeMenuItem = ItemFactory.MakeBoundMenuItem(EditorContext, BoundActionInfo);
+					LastMadeMenuItem = ItemFactory.MakeBoundMenuItem(BoundActionInfo);
 					MenuItemsOut.Add(LastMadeMenuItem);
 
 					if (Flags & FBlueprintActionMenuBuilder::FlattenCategoryHierarcy)
@@ -408,16 +397,12 @@ void FBlueprintActionMenuBuilderImpl::FMenuSectionDefinition::AddBoundMenuItems(
 				}
 			}
 			CompatibleBindings.Empty(); // do before we copy back over cached fields for DatabaseActionInfo
-
-			// copy over any fields that got cached for filtering (with
-			// an empty binding set)
-			/*DatabaseActionInfo = FBlueprintActionInfo(BoundActionInfo, CompatibleBindings);*/
 		}
 	}
 }
 
 //------------------------------------------------------------------------------
-FBlueprintActionMenuBuilderImpl::MenuItemList FBlueprintActionMenuBuilderImpl::FMenuSectionDefinition::MakeMenuItems(TWeakPtr<FBlueprintEditor> EditorContext, FBlueprintActionInfo& DatabaseAction)
+FBlueprintActionMenuBuilderImpl::MenuItemList FBlueprintActionMenuBuilderImpl::FMenuSectionDefinition::MakeMenuItems(FBlueprintActionInfo& DatabaseAction)
 {
 	TRACE_CPUPROFILER_EVENT_SCOPE(FMenuSectionDefinition::MakeMenuItems);
 
@@ -461,7 +446,7 @@ FBlueprintActionMenuBuilderImpl::MenuItemList FBlueprintActionMenuBuilderImpl::F
 
 	if (!UnBoundMenuEntry.IsValid() && bPassedFilter)
 	{
-		UnBoundMenuEntry = ItemFactory.MakeActionMenuItem(EditorContext, DatabaseAction);
+		UnBoundMenuEntry = ItemFactory.MakeActionMenuItem(DatabaseAction);
 		if (Flags & FBlueprintActionMenuBuilder::FlattenCategoryHierarcy)
 		{
 			UnBoundMenuEntry->CosmeticUpdateCategory( ItemFactory.RootCategory );
@@ -473,7 +458,7 @@ FBlueprintActionMenuBuilderImpl::MenuItemList FBlueprintActionMenuBuilderImpl::F
 	{
 		MenuItems.Add(UnBoundMenuEntry);
 	}
-	AddBoundMenuItems(EditorContext, DatabaseAction, GetBindingCandidates(Filter.Context), MenuItems);
+	AddBoundMenuItems(DatabaseAction, Filter.Context.SelectedObjects, MenuItems);
 
 	return MenuItems;
 }
@@ -484,21 +469,17 @@ void FBlueprintActionMenuBuilderImpl::FMenuSectionDefinition::Empty()
 	ConsolidatedProperties.Empty();
 }
 
-//------------------------------------------------------------------------------
-static TArray<FFieldVariant> FBlueprintActionMenuBuilderImpl::GetBindingCandidates(FBlueprintActionContext const& Context)
-{
-	return Context.SelectedObjects;
-}
-
 /*******************************************************************************
  * FBlueprintActionMenuBuilder
  ******************************************************************************/
 
 //------------------------------------------------------------------------------
+// @todo_deprecated - Remove this definition along w/ its declaration in a future release.
+PRAGMA_DISABLE_DEPRECATION_WARNINGS
 FBlueprintActionMenuBuilder::FBlueprintActionMenuBuilder(TWeakPtr<FBlueprintEditor> InBlueprintEditorPtr)
-	: BlueprintEditorPtr(InBlueprintEditorPtr)
 {
 }
+PRAGMA_ENABLE_DEPRECATION_WARNINGS
 
 //------------------------------------------------------------------------------
 void FBlueprintActionMenuBuilder::Empty()
@@ -536,6 +517,10 @@ void FBlueprintActionMenuBuilder::RebuildActionList()
 	
 	FBlueprintActionDatabase& ActionDatabase = FBlueprintActionDatabase::Get();
 	FBlueprintActionDatabase::FActionRegistry const& ActionRegistry = ActionDatabase.GetAllActions();
+
+#if ENABLE_BLUEPRINT_ACTION_FILTER_PROFILING
+	int32 TotalActionCount = 0;
+#endif	// ENABLE_BLUEPRINT_ACTION_FILTER_PROFILING
 	
 	for (auto Iterator(ActionRegistry.CreateConstIterator()); Iterator; ++Iterator)
 	{
@@ -548,11 +533,13 @@ void FBlueprintActionMenuBuilder::RebuildActionList()
 			{
 				FBlueprintActionInfo BlueprintAction(ActionObject, NodeSpawner);
 
-				// @TODO: could probably have a super filter that spreads across 
-				//        all MenuSctions (to pair down on performance?)
+#if ENABLE_BLUEPRINT_ACTION_FILTER_PROFILING
+				++TotalActionCount;
+#endif	// ENABLE_BLUEPRINT_ACTION_FILTER_PROFILING
+
 				for (TSharedRef<FMenuSectionDefinition> const& MenuSection : MenuSections)
 				{
-					for (TSharedPtr<FEdGraphSchemaAction> MenuEntry : MenuSection->MakeMenuItems(BlueprintEditorPtr, BlueprintAction))
+					for (TSharedPtr<FEdGraphSchemaAction> MenuEntry : MenuSection->MakeMenuItems(BlueprintAction))
 					{
 						AddAction(MenuEntry);
 					}
@@ -567,12 +554,22 @@ void FBlueprintActionMenuBuilder::RebuildActionList()
 	}
 
 #if ENABLE_BLUEPRINT_ACTION_FILTER_PROFILING
-	if (CVarBPEnableActionMenuFilterTestStatsLogging.GetValueOnGameThread())
+	if (FBlueprintActionFilter::IsFilterTestStatsLoggingEnabled())
 	{
+		UE_LOG(LogBlueprintActionMenuItemFactory, Log, TEXT("=== UNFILTERED ACTIONS: %d"), TotalActionCount);
+
 		// Dump detailed stats information about each filter test that was involved with building each menu section.
-		for (const TSharedRef<FMenuSectionDefinition>& MenuSection : MenuSections)
+		for (int32 SectionIdx = 0; SectionIdx < MenuSections.Num(); ++SectionIdx)
 		{
-			UE_LOG(LogBlueprintActionMenuItemFactory, Log, TEXT("=== %s: FILTER TEST PROFILE ==="), *MenuSection->GetSectionHeading().ToString());
+			const TSharedRef<FMenuSectionDefinition>& MenuSection = MenuSections[SectionIdx];
+			FString DisplayName = FString::Printf(TEXT("MenuSection[%d]"), SectionIdx);
+			const FText& SectionHeading = MenuSection->GetSectionHeading();
+			if (!SectionHeading.IsEmptyOrWhitespace())
+			{
+				DisplayName.Appendf(TEXT(" (%s)"), *SectionHeading.ToString());
+			}
+
+			UE_LOG(LogBlueprintActionMenuItemFactory, Log, TEXT("=== FILTER TEST PROFILE: %s ==="), *DisplayName);
 
 			for (const FString& ProfileEntry : MenuSection->Filter.GetFilterTestProfile())
 			{
@@ -580,6 +577,43 @@ void FBlueprintActionMenuBuilder::RebuildActionList()
 			}
 		}
 	}
+
+	if (CVarBPEnableActionMenuDumpToFile.GetValueOnGameThread())
+	{
+		TStringBuilder<512> CurrentLine;
+		TStringBuilder<256> CategoryPrefix;
+
+		FString DumpFilePath = FPaths::ProfilingDir() / TEXT("ActionMenu") / FString::Printf(TEXT("MenuDump_%s.txt"), *FDateTime::Now().ToString());
+		FOutputDeviceFile DumpFile(*DumpFilePath, true);
+
+		DumpFile.SetSuppressEventTag(true);
+
+		UE_LOG(LogBlueprintActionMenuItemFactory, Log, TEXT("Dumping menu builder action list to \"%s\"..."), *DumpFilePath);
+
+		int32 TotalCount = 0;
+
+		for (int32 i = 0; i < GetNumActions(); ++i)
+		{
+			const ActionGroup& CurrentEntry = GetAction(i);
+
+			CategoryPrefix.Reset();
+			for (const FString& Category : CurrentEntry.GetCategoryChain())
+			{
+				CategoryPrefix.Append(FString::Printf(TEXT("%s|"), *Category));
+			}
+
+			for (const TSharedPtr<FEdGraphSchemaAction>& Action : CurrentEntry.Actions)
+			{
+				CurrentLine = FString::Printf(TEXT("%05d "), ++TotalCount);
+				CurrentLine.Append(CategoryPrefix.ToString());
+				CurrentLine.Append(Action->GetMenuDescription().ToString());
+
+				DumpFile.Log(CurrentLine.ToString());
+			}
+		}
+
+		CVarBPEnableActionMenuDumpToFile.AsVariable()->Set(false);
+	}
 #endif	// ENABLE_BLUEPRINT_ACTION_FILTER_PROFILING
 }
 
diff --git a/Engine/Source/Editor/Kismet/Private/BlueprintActionMenuUtils.cpp b/Engine/Source/Editor/Kismet/Private/BlueprintActionMenuUtils.cpp
--- a/Engine/Source/Editor/Kismet/Private/BlueprintActionMenuUtils.cpp
+++ b/Engine/Source/Editor/Kismet/Private/BlueprintActionMenuUtils.cpp
@@ -314,7 +314,7 @@ static FBlueprintActionFilter BlueprintActionMenuUtilsImpl::MakeCallOnMemberFilt
 	FBlueprintActionFilter CallOnMemberFilter;
 	CallOnMemberFilter.Context = MainMenuFilter.Context;
 	CallOnMemberFilter.PermittedNodeTypes.Add(UK2Node_CallFunction::StaticClass());
-	CallOnMemberFilter.AddRejectionTest(MAKE_ACTION_FILTER_REJECTION_TEST(IsUnBoundSpawner)->WithFlags(EActionFilterTestFlags::CacheResults));
+	CallOnMemberFilter.AddRejectionTest(FBlueprintActionFilter::FRejectionTestDelegate::CreateStatic(IsUnBoundSpawner));
 	
 	const UBlueprintEditorSettings* BlueprintSettings = GetDefault<UBlueprintEditorSettings>();
 	// instead of looking for "ExposeFunctionCategories" on component properties,
@@ -322,11 +322,11 @@ static FBlueprintActionFilter BlueprintActionMenuUtilsImpl::MakeCallOnMemberFilt
 	// for "ExposeFunctionCategories" on any non-component properties... 
 	if (BlueprintSettings->bExposeAllMemberComponentFunctions)
 	{
-		CallOnMemberFilter.AddRejectionTest(MAKE_ACTION_FILTER_REJECTION_TEST(IsUnexposedNonComponentAction)->WithFlags(EActionFilterTestFlags::CacheResults));
+		CallOnMemberFilter.AddRejectionTest(FBlueprintActionFilter::FRejectionTestDelegate::CreateStatic(IsUnexposedNonComponentAction));
 	}
 	else
 	{
-		CallOnMemberFilter.AddRejectionTest(MAKE_ACTION_FILTER_REJECTION_TEST(IsUnexposedMemberAction)->WithFlags(EActionFilterTestFlags::CacheResults));
+		CallOnMemberFilter.AddRejectionTest(FBlueprintActionFilter::FRejectionTestDelegate::CreateStatic(IsUnexposedMemberAction));
 	}
 
 	bool bForceAddComponents = ((ContextTargetMask & EContextTargetFlags::TARGET_SubComponents) != 0);
@@ -423,8 +423,7 @@ static void BlueprintActionMenuUtilsImpl::AddFavoritesSection(FBlueprintActionFi
 	{
 		FBlueprintActionFilter FavoritesFilter = MainMenuFilter;
 
-		// make this action filter uncached because favorites can change on a moments notice
-		FavoritesFilter.AddRejectionTest(MAKE_ACTION_FILTER_REJECTION_TEST(IsNonFavoritedAction));
+		FavoritesFilter.AddRejectionTest(FBlueprintActionFilter::FRejectionTestDelegate::CreateStatic(IsNonFavoritedAction));
 		
 		uint32 SectionFlags = 0x00;
 		FText  SectionHeading = LOCTEXT("ContextMenuFavoritesTitle", "Favorites");
@@ -509,18 +508,18 @@ void FBlueprintActionMenuUtils::MakeContextMenu(FBlueprintActionContext const& C
 	FBlueprintActionFilter ComponentsFilter;
 	ComponentsFilter.Context = Context;
 	// only want bound actions for this menu section
-	ComponentsFilter.AddRejectionTest(MAKE_ACTION_FILTER_REJECTION_TEST(IsUnBoundSpawner)->WithFlags(EActionFilterTestFlags::CacheResults));
+	ComponentsFilter.AddRejectionTest(FBlueprintActionFilter::FRejectionTestDelegate::CreateStatic(IsUnBoundSpawner));
 	// @TODO: don't know exactly why we can only bind non-pure/const functions;
 	//        this is mirrored after FK2ActionMenuBuilder::GetFunctionCallsOnSelectedActors()
 	//        and FK2ActionMenuBuilder::GetFunctionCallsOnSelectedComponents(),
 	//        where we make the same stipulation
-	ComponentsFilter.AddRejectionTest(MAKE_ACTION_FILTER_REJECTION_TEST(IsPureNonConstAction)->WithFlags(EActionFilterTestFlags::CacheResults));
+	ComponentsFilter.AddRejectionTest(FBlueprintActionFilter::FRejectionTestDelegate::CreateStatic(IsPureNonConstAction));
 	
 
 	FBlueprintActionFilter LevelActorsFilter;
 	LevelActorsFilter.Context = Context;
 	// only want bound actions for this menu section
-	LevelActorsFilter.AddRejectionTest(MAKE_ACTION_FILTER_REJECTION_TEST(IsUnBoundSpawner)->WithFlags(EActionFilterTestFlags::CacheResults));
+	LevelActorsFilter.AddRejectionTest(FBlueprintActionFilter::FRejectionTestDelegate::CreateStatic(IsUnBoundSpawner));
 
 	// Build asset reference filter
 	FAssetReferenceFilterContext AssetReferenceFilterContext;
@@ -672,7 +671,7 @@ void FBlueprintActionMenuUtils::MakeContextMenu(FBlueprintActionContext const& C
 	FBlueprintActionFilter AddComponentFilter;
 	AddComponentFilter.Context = MainMenuFilter.Context;
 	AddComponentFilter.PermittedNodeTypes.Add(UK2Node_AddComponent::StaticClass());
-	AddComponentFilter.AddRejectionTest(MAKE_ACTION_FILTER_REJECTION_TEST(IsUnBoundSpawner)->WithFlags(EActionFilterTestFlags::CacheResults));
+	AddComponentFilter.AddRejectionTest(FBlueprintActionFilter::FRejectionTestDelegate::CreateStatic(IsUnBoundSpawner));
 
 
 	FContentBrowserModule& ContentBrowserModule = FModuleManager::LoadModuleChecked<FContentBrowserModule>("ContentBrowser");
@@ -775,7 +774,7 @@ void FBlueprintActionMenuUtils::MakeFavoritesMenu(FBlueprintActionContext const&
 
 	FBlueprintActionFilter MenuFilter;
 	MenuFilter.Context = Context;
-	MenuFilter.AddRejectionTest(MAKE_ACTION_FILTER_REJECTION_TEST(BlueprintActionMenuUtilsImpl::IsNonFavoritedAction));
+	MenuFilter.AddRejectionTest(FBlueprintActionFilter::FRejectionTestDelegate::CreateStatic(BlueprintActionMenuUtilsImpl::IsNonFavoritedAction));
 
 	uint32 SectionFlags = 0x00;
 	const UBlueprintEditorSettings* BlueprintSettings = GetDefault<UBlueprintEditorSettings>();
diff --git a/Engine/Source/Editor/Kismet/Private/SBlueprintActionMenu.cpp b/Engine/Source/Editor/Kismet/Private/SBlueprintActionMenu.cpp
--- a/Engine/Source/Editor/Kismet/Private/SBlueprintActionMenu.cpp
+++ b/Engine/Source/Editor/Kismet/Private/SBlueprintActionMenu.cpp
@@ -521,7 +521,7 @@ void SBlueprintActionMenu::CollectAllActions(FGraphActionListBuilderBase& OutAll
 	FBlueprintActionContext FilterContext;
 	ConstructActionContext(FilterContext);
 	
-	FBlueprintActionMenuBuilder MenuBuilder(EditorPtr);
+	FBlueprintActionMenuBuilder MenuBuilder;
 	// NOTE: cannot call GetGraphContextActions() during serialization and GC due to its use of FindObject()
 	if(!GIsSavingPackage && !IsGarbageCollecting() && FilterContext.Blueprints.Num() > 0)
 	{
diff --git a/Engine/Source/Editor/Kismet/Private/SBlueprintFavoritesPalette.cpp b/Engine/Source/Editor/Kismet/Private/SBlueprintFavoritesPalette.cpp
--- a/Engine/Source/Editor/Kismet/Private/SBlueprintFavoritesPalette.cpp
+++ b/Engine/Source/Editor/Kismet/Private/SBlueprintFavoritesPalette.cpp
@@ -346,7 +346,7 @@ void SBlueprintFavoritesPalette::CollectAllActions(FGraphActionListBuilderBase&
 	FilterContext.Blueprints.Add(GetBlueprint());
 	FilterContext.EditorPtr = BlueprintEditorPtr;
 
-	FBlueprintActionMenuBuilder FavoritesBuilder(BlueprintEditorPtr);
+	FBlueprintActionMenuBuilder FavoritesBuilder;
 	FBlueprintActionMenuUtils::MakeFavoritesMenu(FilterContext, FavoritesBuilder);
 
 	OutAllActions.Append(FavoritesBuilder);
diff --git a/Engine/Source/Editor/Kismet/Private/SBlueprintLibraryPalette.cpp b/Engine/Source/Editor/Kismet/Private/SBlueprintLibraryPalette.cpp
--- a/Engine/Source/Editor/Kismet/Private/SBlueprintLibraryPalette.cpp
+++ b/Engine/Source/Editor/Kismet/Private/SBlueprintLibraryPalette.cpp
@@ -257,7 +257,7 @@ void SBlueprintLibraryPalette::CollectAllActions(FGraphActionListBuilderBase& Ou
 		ClassFilter = FilterClass.Get();
 	}
 	
-	FBlueprintActionMenuBuilder PaletteBuilder(BlueprintEditorPtr);
+	FBlueprintActionMenuBuilder PaletteBuilder;
 	FBlueprintActionMenuUtils::MakePaletteMenu(FilterContext, ClassFilter, PaletteBuilder);
 	OutAllActions.Append(PaletteBuilder);
 }
diff --git a/Engine/Source/Editor/Kismet/Public/BlueprintActionMenuBuilder.h b/Engine/Source/Editor/Kismet/Public/BlueprintActionMenuBuilder.h
--- a/Engine/Source/Editor/Kismet/Public/BlueprintActionMenuBuilder.h
+++ b/Engine/Source/Editor/Kismet/Public/BlueprintActionMenuBuilder.h
@@ -46,6 +46,9 @@ public:
 	};
 	
 public:
+	FBlueprintActionMenuBuilder() = default;
+
+	UE_DEPRECATED(5.2, "The action filter now stores a reference to the authoritative editor context. Please use the default constructor instead.")
 	FBlueprintActionMenuBuilder(TWeakPtr<FBlueprintEditor> BlueprintEditorPtr);
 	
 	// FGraphActionListBuilderBase interface
@@ -80,7 +83,4 @@ private:
 	 * this API clean).
 	 */
 	TArray< TSharedRef<FBlueprintActionMenuBuilderImpl::FMenuSectionDefinition> > MenuSections;
-
-	/** */
-	TWeakPtr<FBlueprintEditor> BlueprintEditorPtr;
 };
diff --git a/Engine/Source/Editor/Kismet/Public/BlueprintEditor.h b/Engine/Source/Editor/Kismet/Public/BlueprintEditor.h
--- a/Engine/Source/Editor/Kismet/Public/BlueprintEditor.h
+++ b/Engine/Source/Editor/Kismet/Public/BlueprintEditor.h
@@ -276,6 +276,7 @@ public:
 	virtual void JumpToPin(const class UEdGraphPin* Pin) override;
 	virtual void SummonSearchUI(bool bSetFindWithinBlueprint, FString NewSearchTerms = FString(), bool bSelectFirstResult = false) override;
 	virtual void SummonFindAndReplaceUI() override;
+	virtual UEdGraph* GetFocusedGraph() const override;
 	virtual TSharedPtr<SGraphEditor> OpenGraphAndBringToFront(UEdGraph* Graph, bool bSetFocus = true) override;
 
 	UE_DEPRECATED(5.0, "GetSelectedSCSEditorTreeNodes has been deprecated. Use GetSelectedSubobjectEditorTreeNodes instead.")
@@ -723,11 +724,6 @@ public:
 		return CustomizeSubobjectEditor(InComponentToCustomize);
 	}
 
-	/**
-	 * Returns the currently focused graph in the Blueprint editor
-	 */
-	UEdGraph* GetFocusedGraph() const;
-
 	/** Adds to a list of custom objects for debugging beyond what will automatically be found/used */
 	virtual void GetCustomDebugObjects(TArray<FCustomDebugObject>& DebugList) const { }
 
diff --git a/Engine/Source/Editor/Kismet/Public/BlueprintEditorModule.h b/Engine/Source/Editor/Kismet/Public/BlueprintEditorModule.h
--- a/Engine/Source/Editor/Kismet/Public/BlueprintEditorModule.h
+++ b/Engine/Source/Editor/Kismet/Public/BlueprintEditorModule.h
@@ -99,6 +99,9 @@ public:
 	/** Invokes the Find and Replace UI */
 	virtual void SummonFindAndReplaceUI() = 0;
 
+	/** Returns the currently focused graph context, or NULL if no graph is available. */
+	virtual UEdGraph* GetFocusedGraph() const = 0;
+
 	/** Tries to open the specified graph and bring it's document to the front (note: this can return NULL) */
 	virtual TSharedPtr<class SGraphEditor> OpenGraphAndBringToFront(class UEdGraph* Graph, bool bSetFocus = true) = 0;
 
diff --git a/Engine/Source/Editor/MovieSceneTools/Private/MovieSceneEventCustomization.cpp b/Engine/Source/Editor/MovieSceneTools/Private/MovieSceneEventCustomization.cpp
--- a/Engine/Source/Editor/MovieSceneTools/Private/MovieSceneEventCustomization.cpp
+++ b/Engine/Source/Editor/MovieSceneTools/Private/MovieSceneEventCustomization.cpp
@@ -67,7 +67,7 @@ namespace UE_MovieSceneEventCustomization
 		};
 
 
-		FBlueprintActionMenuBuilder ContextMenuBuilder(nullptr);
+		FBlueprintActionMenuBuilder ContextMenuBuilder;
 
 		if (BoundObjectPinClass)
 		{
@@ -91,7 +91,7 @@ namespace UE_MovieSceneEventCustomization
 			// but also prevents displaying functions on BP components. Comment out for now.
 			//CallOnMemberFilter.AddRejectionTest(FBlueprintActionFilter::FRejectionTestDelegate::CreateStatic(RejectAnyUnboundActions));
 			
-			CallOnMemberFilter.AddRejectionTest(MAKE_ACTION_FILTER_REJECTION_TEST(RejectAnyNonFunctions)->WithFlags(EActionFilterTestFlags::CacheResults));
+			CallOnMemberFilter.AddRejectionTest(FBlueprintActionFilter::FRejectionTestDelegate::CreateStatic(RejectAnyNonFunctions));
 
 			ContextMenuBuilder.AddMenuSection(CallOnMemberFilter, FText::FromName(BoundObjectPinClass->GetFName()), 0);
 		}
@@ -110,7 +110,7 @@ namespace UE_MovieSceneEventCustomization
 			{
 				FBlueprintActionFilter::AddUnique(MenuFilter.TargetClasses, Blueprint->SkeletonGeneratedClass);
 			}
-			MenuFilter.AddRejectionTest(MAKE_ACTION_FILTER_REJECTION_TEST(RejectAnyNonFunctions)->WithFlags(EActionFilterTestFlags::CacheResults));
+			MenuFilter.AddRejectionTest(FBlueprintActionFilter::FRejectionTestDelegate::CreateStatic(RejectAnyNonFunctions));
 
 
 			ContextMenuBuilder.AddMenuSection(MenuFilter, LOCTEXT("SequenceDirectorMenu", "This Sequence"), 0);
@@ -143,14 +143,14 @@ namespace UE_MovieSceneEventCustomization
 
 		MenuFilter.Context.Blueprints.Add(Blueprint);
 
-		MenuFilter.AddRejectionTest(MAKE_ACTION_FILTER_REJECTION_TEST(RejectAnyForeignFunctions, Blueprint));
+		MenuFilter.AddRejectionTest(FBlueprintActionFilter::FRejectionTestDelegate::CreateStatic(RejectAnyForeignFunctions, Blueprint));
 
 		if (Blueprint->SkeletonGeneratedClass)
 		{
 			FBlueprintActionFilter::AddUnique(MenuFilter.TargetClasses, Blueprint->SkeletonGeneratedClass);
 		}
 
-		FBlueprintActionMenuBuilder ContextMenuBuilder(nullptr);
+		FBlueprintActionMenuBuilder ContextMenuBuilder;
 
 		ContextMenuBuilder.AddMenuSection(MenuFilter, LOCTEXT("SequenceDirectorMenu", "This Sequence"), 0);
 		ContextMenuBuilder.RebuildActionList();
diff --git a/Engine/Source/Editor/UnrealEd/Private/Commandlets/GenerateBlueprintAPICommandlet.cpp b/Engine/Source/Editor/UnrealEd/Private/Commandlets/GenerateBlueprintAPICommandlet.cpp
--- a/Engine/Source/Editor/UnrealEd/Private/Commandlets/GenerateBlueprintAPICommandlet.cpp
+++ b/Engine/Source/Editor/UnrealEd/Private/Commandlets/GenerateBlueprintAPICommandlet.cpp
@@ -678,7 +678,7 @@ static double GenerateBlueprintAPIUtils::GetPaletteMenuActions(FCategorizedGraph
 	FBlueprintActionContext FilterContext;
 	FilterContext.Blueprints.Add(const_cast<UBlueprint*>(Blueprint));
 
-	FBlueprintActionMenuBuilder MenuBuilder(nullptr);
+	FBlueprintActionMenuBuilder MenuBuilder;
 	{
 		// prime the database so it's not recorded in our timing capture
 		FBlueprintActionDatabase::Get();
