Optimized CreateOpTable by updating the table in parallel.
Also removed some unnecessary string allocations in FTypePromotion::GetOpNameFromFunction.

Before: 3s 357ms
After: 1s 674ms

diff --git a/Engine/Source/Editor/BlueprintGraph/Private/BlueprintTypePromotion.cpp b/Engine/Source/Editor/BlueprintGraph/Private/BlueprintTypePromotion.cpp
--- a/Engine/Source/Editor/BlueprintGraph/Private/BlueprintTypePromotion.cpp
+++ b/Engine/Source/Editor/BlueprintGraph/Private/BlueprintTypePromotion.cpp
@@ -483,11 +483,13 @@ FName FTypePromotion::GetOpNameFromFunction(UFunction const* const Func)
 		return OperatorNames::NoOp;
 	}
 
-	FString FuncName = Func->GetName();
+	TStringBuilder<256> FuncName;
+	Func->GetFName().ToString(FuncName);
+	TStringView FuncNameView = FuncName.ToView();
 	// Get everything before the "_"
-	int32 Index = FuncName.Find(TEXT("_"));
+	int32 Index = FuncNameView.Find(TEXT("_"));
 	
-	FName FuncNameChopped(FuncName.Mid(0, Index));
+	FName FuncNameChopped(FuncNameView.Mid(0, Index));
 	if (GetAllOpNames().Contains(FuncNameChopped))
 	{
 		return FuncNameChopped;
@@ -500,35 +502,39 @@ void FTypePromotion::CreateOpTable()
 {
 	TRACE_CPUPROFILER_EVENT_SCOPE(FTypePromotion::CreateOpTable);
 	const UEdGraphSchema_K2* Schema = GetDefault<UEdGraphSchema_K2>();
+	FCriticalSection AddOpLock;
 
 	OperatorTable.Empty();
 
 	TArray<UClass*> Libraries;
 	GetDerivedClasses(UBlueprintFunctionLibrary::StaticClass(), Libraries);
-	for (UClass* Library : Libraries)
-	{
-		// Ignore abstract libraries/classes
-		if (!Library || Library->HasAnyClassFlags(CLASS_Abstract))
-		{
-			continue;
-		}
-
-		for (UFunction* Function : TFieldRange<UFunction>(Library, EFieldIteratorFlags::ExcludeSuper, EFieldIteratorFlags::ExcludeDeprecated))
+	ParallelFor(Libraries.Num(), 
+		[this, &Schema, &AddOpLock, &Libraries](int32 Index)
 		{
-			if(!IsPromotableFunction(Function))
+			UClass* Library = Libraries[Index];
+			// Ignore abstract libraries/classes
+			if (!Library || Library->HasAnyClassFlags(CLASS_Abstract))
 			{
-				continue;
+				return;
 			}
 
-			FEdGraphPinType FuncPinType;
-			FName OpName = GetOpNameFromFunction(Function);
-
-			if (OpName != OperatorNames::NoOp && Schema->ConvertPropertyToPinType(Function->GetReturnProperty(), /* out */ FuncPinType))
+			for (UFunction* Function : TFieldRange<UFunction>(Library, EFieldIteratorFlags::ExcludeSuper, EFieldIteratorFlags::ExcludeDeprecated))
 			{
-				AddOpFunction(OpName, Function);
+				if (!IsPromotableFunction(Function))
+				{
+					continue;
+				}
+
+				FEdGraphPinType FuncPinType;
+				FName OpName = GetOpNameFromFunction(Function);
+
+				if (OpName != OperatorNames::NoOp && Schema->ConvertPropertyToPinType(Function->GetReturnProperty(), /* out */ FuncPinType))
+				{
+					FScopeLock ScopeLock(&AddOpLock);
+					AddOpFunction(OpName, Function);
+				}
 			}
-		}
-	}
+		});
 }
 
 void FTypePromotion::AddOpFunction(FName OpName, UFunction* Function)
