Fixed non-unity build CIS issue.

diff --git a/Engine/Source/Editor/BlueprintGraph/Private/BlueprintTypePromotion.cpp b/Engine/Source/Editor/BlueprintGraph/Private/BlueprintTypePromotion.cpp
--- a/Engine/Source/Editor/BlueprintGraph/Private/BlueprintTypePromotion.cpp
+++ b/Engine/Source/Editor/BlueprintGraph/Private/BlueprintTypePromotion.cpp
@@ -2,6 +2,7 @@
 
 #include "BlueprintTypePromotion.h"
 
+#include "Async/ParallelFor.h"
 #include "Containers/EnumAsByte.h"
 #include "Containers/UnrealString.h"
 #include "Delegates/Delegate.h"
