UE-171455 - Cannot find appropriate promotion

Possible fallout from CL 23313465.

Some promotion operator pins can have single precision floats. When these are linked to double precision inputs, we need to permit the connection without the need of a cast node. This is achieved by permitting real number connections in CreateIntermediateCast.

diff --git a/Engine/Source/Editor/BlueprintGraph/Private/K2Node_PromotableOperator.cpp b/Engine/Source/Editor/BlueprintGraph/Private/K2Node_PromotableOperator.cpp
--- a/Engine/Source/Editor/BlueprintGraph/Private/K2Node_PromotableOperator.cpp
+++ b/Engine/Source/Editor/BlueprintGraph/Private/K2Node_PromotableOperator.cpp
@@ -976,8 +976,13 @@ bool UK2Node_PromotableOperator::CreateIntermediateCast(UK2Node_CallFunction* So
 	check(InputPin && OutputPin);
 	const UEdGraphSchema_K2* Schema = CompilerContext.GetSchema();
 
-	// If the pin types are the same, than no casts are needed and we can just connect
-	if (InputPin->PinType == OutputPin->PinType)
+	// If the pin types are the same, than no casts are needed and we can just connect.
+	// This includes real number types that may have a different subcategory, which we implicitly convert.
+	const bool bPinTypesMatch =
+		(InputPin->PinType == OutputPin->PinType) ||
+		((InputPin->PinType.PinCategory == UEdGraphSchema_K2::PC_Real) && (OutputPin->PinType.PinCategory == UEdGraphSchema_K2::PC_Real));
+
+	if (bPinTypesMatch)
 	{
 		// If SourceNode is 'this' then we need to move the pin links instead of just 
 		// creating the connection because the output is not another new node, but 
