Misc Blueprint changes:

* Added overloads for FindSpecializedConversionNode and SearchForAutocastFunction that don't use output parameters. There are several instances in the engine where dummy variables are used to get around this.
* Updated the Message_ functions to use template parameter packing. The arguments are forwarded to the underlying FCompilerResultsLog, if one exists. This allows us to use tokenized messages.
* Changed several functions in FCompilerResultsLog to use the coding standard's naming convention. Specifically, arguments need to start with an uppercase letter.

diff --git a/Engine/Source/Editor/AnimGraph/Private/AnimBlueprintExtension_PropertyAccess.cpp b/Engine/Source/Editor/AnimGraph/Private/AnimBlueprintExtension_PropertyAccess.cpp
--- a/Engine/Source/Editor/AnimGraph/Private/AnimBlueprintExtension_PropertyAccess.cpp
+++ b/Engine/Source/Editor/AnimGraph/Private/AnimBlueprintExtension_PropertyAccess.cpp
@@ -397,12 +397,10 @@ void UAnimBlueprintExtension_PropertyAccess::ExpandPropertyAccess(FKismetCompile
 			else
 			{
 				// Need to create a conversion node. We will support basic conversion here for now as we are only looking for primitive types/casts
-				FName FuncName = NAME_None;
-				UClass* FuncOwner = nullptr;
-				if(K2Schema->SearchForAutocastFunction(CurrentPin->PinType, InTargetPin->PinType, FuncName, FuncOwner))
+				if(TOptional<UEdGraphSchema_K2::FSearchForAutocastFunctionResults> AutoCastResults = K2Schema->SearchForAutocastFunction(CurrentPin->PinType, InTargetPin->PinType))
 				{
 					UK2Node_CallFunction* AutoCastNode = InCompilerContext.SpawnIntermediateNode<UK2Node_CallFunction>(SourceNode, InParentGraph);
-					AutoCastNode->FunctionReference.SetExternalMember(FuncName, FuncOwner);
+					AutoCastNode->FunctionReference.SetExternalMember(AutoCastResults->TargetFunction, AutoCastResults->FunctionOwner);
 					AutoCastNode->AllocateDefaultPins();
 					
 					// Find output pin & connect
diff --git a/Engine/Source/Editor/AnimGraph/Private/AnimationGraphSchema.cpp b/Engine/Source/Editor/AnimGraph/Private/AnimationGraphSchema.cpp
--- a/Engine/Source/Editor/AnimGraph/Private/AnimationGraphSchema.cpp
+++ b/Engine/Source/Editor/AnimGraph/Private/AnimationGraphSchema.cpp
@@ -518,20 +518,37 @@ FReply UAnimationGraphSchema::BeginGraphDragAction(TSharedPtr<FEdGraphSchemaActi
 
 bool UAnimationGraphSchema::SearchForAutocastFunction(const FEdGraphPinType& OutputPinType, const FEdGraphPinType& InputPinType, FName& TargetFunction, /*out*/ UClass*& FunctionOwner) const
 {
+	TOptional<UEdGraphSchema_K2::FSearchForAutocastFunctionResults> Result = SearchForAutocastFunction(OutputPinType, InputPinType);
+	if (Result)
+	{
+		TargetFunction = Result->TargetFunction;
+		FunctionOwner = Result->FunctionOwner;
+		return true;
+	}
+
+	return false;
+}
+
+TOptional<UEdGraphSchema_K2::FSearchForAutocastFunctionResults> UAnimationGraphSchema::SearchForAutocastFunction(const FEdGraphPinType& OutputPinType, const FEdGraphPinType& InputPinType) const
+{
+	TOptional<UEdGraphSchema_K2::FSearchForAutocastFunctionResults> Result;
+
 	if (IsComponentSpacePosePin(OutputPinType) && IsLocalSpacePosePin(InputPinType))
 	{
 		// Insert a Component To LocalSpace conversion
-		return true;
+		Result = UEdGraphSchema_K2::FSearchForAutocastFunctionResults{};
 	}
 	else if (IsLocalSpacePosePin(OutputPinType) && IsComponentSpacePosePin(InputPinType))
 	{
 		// Insert a Local To ComponentSpace conversion
-		return true;
+		Result = UEdGraphSchema_K2::FSearchForAutocastFunctionResults{};
 	}
 	else
 	{
-		return Super::SearchForAutocastFunction(OutputPinType, InputPinType, TargetFunction, FunctionOwner);
+		Result = Super::SearchForAutocastFunction(OutputPinType, InputPinType);
 	}
+
+	return Result;
 }
 
 bool UAnimationGraphSchema::CreateAutomaticConversionNodeAndConnections(UEdGraphPin* PinA, UEdGraphPin* PinB) const
diff --git a/Engine/Source/Editor/AnimGraph/Public/AnimationGraphSchema.h b/Engine/Source/Editor/AnimGraph/Public/AnimationGraphSchema.h
--- a/Engine/Source/Editor/AnimGraph/Public/AnimationGraphSchema.h
+++ b/Engine/Source/Editor/AnimGraph/Public/AnimationGraphSchema.h
@@ -71,7 +71,9 @@ class UAnimationGraphSchema : public UEdGraphSchema_K2
 
 	//~ Begin UEdGraphSchema_K2 Interface
 	virtual const FPinConnectionResponse DetermineConnectionResponseOfCompatibleTypedPins(const UEdGraphPin* PinA, const UEdGraphPin* PinB, const UEdGraphPin* InputPin, const UEdGraphPin* OutputPin) const override;
+	UE_DEPRECATED(5.2, "Use the FSearchForAutocastFunctionResults variant.")
 	virtual bool SearchForAutocastFunction(const FEdGraphPinType& OutputPinType, const FEdGraphPinType& InputPinType, /*out*/ FName& TargetFunction, /*out*/ UClass*& FunctionOwner) const override;
+	virtual TOptional<FSearchForAutocastFunctionResults> SearchForAutocastFunction(const FEdGraphPinType& OutputPinType, const FEdGraphPinType& InputPinType) const;
 	virtual bool ArePinsCompatible(const UEdGraphPin* PinA, const UEdGraphPin* PinB, const UClass* CallingContext = NULL, bool bIgnoreArray = false) const override;
 	virtual bool DoesSupportCollapsedNodes() const override { return false; }
 	virtual bool DoesSupportEventDispatcher() const	override { return false; }
diff --git a/Engine/Source/Editor/BlueprintGraph/Classes/EdGraphSchema_K2.h b/Engine/Source/Editor/BlueprintGraph/Classes/EdGraphSchema_K2.h
--- a/Engine/Source/Editor/BlueprintGraph/Classes/EdGraphSchema_K2.h
+++ b/Engine/Source/Editor/BlueprintGraph/Classes/EdGraphSchema_K2.h
@@ -1114,15 +1114,30 @@ public:
 	/** Find a 'set value by name' function for the specified pin, if it exists */
 	static UFunction* FindSetVariableByNameFunction(const FEdGraphPinType& PinType);
 
-	/** Find an appropriate function to call to perform an automatic cast operation */
+	UE_DEPRECATED(5.2, "Use the FSearchForAutocastFunctionResults variant.")
 	virtual bool SearchForAutocastFunction(const FEdGraphPinType& OutputPinType, const FEdGraphPinType& InputPinType, /*out*/ FName& TargetFunction, /*out*/ UClass*& FunctionOwner) const;
 
-	/** Find an appropriate node that can convert from one pin type to another (not a cast; e.g. "MakeLiteralArray" node) */
+	/** Find an appropriate function to call to perform an automatic cast operation */
+	struct FSearchForAutocastFunctionResults
+	{
+		FName TargetFunction;
+		UClass* FunctionOwner = nullptr;
+	};
+	UE_NODISCARD virtual TOptional<FSearchForAutocastFunctionResults> SearchForAutocastFunction(const FEdGraphPinType& OutputPinType, const FEdGraphPinType& InputPinType) const;
+
+	UE_DEPRECATED(5.2, "Use the FSearchForAutocastFunctionResults variant.")
 	virtual bool FindSpecializedConversionNode(const UEdGraphPin* OutputPin, const UEdGraphPin* InputPin, bool bCreateNode, /*out*/ class UK2Node*& TargetNode) const;
 	
-	/** Find an appropriate node that can convert from one pin type to another (not a cast; e.g. "MakeLiteralArray" node) */
+	UE_DEPRECATED(5.2, "Use the FSearchForAutocastFunctionResults variant.")
 	virtual bool FindSpecializedConversionNode(const FEdGraphPinType& OutputPinType, const UEdGraphPin* InputPinType, bool bCreateNode, /*out*/ class UK2Node*& TargetNode) const;
 	
+	/** Find an appropriate node that can convert from one pin type to another (not a cast; e.g. "MakeLiteralArray" node) */
+	struct FindSpecializedConversionNodeResults
+	{
+		class UK2Node* TargetNode = nullptr;
+	};
+	UE_NODISCARD virtual TOptional<FindSpecializedConversionNodeResults> FindSpecializedConversionNode(const FEdGraphPinType& OutputPinType, const UEdGraphPin& InputPin, bool bCreateNode) const;
+
 	/** Create menu for variable get/set nodes which refer to a variable which does not exist. */
 	void GetNonExistentVariableMenu(FToolMenuSection& Section, const UEdGraphNode* InGraphNode, UBlueprint* OwnerBlueprint) const;
 
diff --git a/Engine/Source/Editor/BlueprintGraph/Classes/K2Node.h b/Engine/Source/Editor/BlueprintGraph/Classes/K2Node.h
--- a/Engine/Source/Editor/BlueprintGraph/Classes/K2Node.h
+++ b/Engine/Source/Editor/BlueprintGraph/Classes/K2Node.h
@@ -11,6 +11,7 @@
 #include "CoreMinimal.h"
 #include "Delegates/Delegate.h"
 #include "EdGraph/EdGraphNode.h"
+#include "Engine/Blueprint.h"
 #include "HAL/PlatformMath.h"
 #include "Internationalization/Text.h"
 #include "Math/Color.h"
@@ -491,10 +492,47 @@ protected:
 	 * Sends a message to the owning blueprint's CurrentMessageLog, if there is one available.  Otherwise, defaults to logging to the normal channels.
 	 * Should use this for node actions that happen during compilation!
 	 */
-	void Message_Note(const FString& Message);
-	void Message_Warn(const FString& Message);
-	void Message_Error(const FString& Message);
+	template<typename... ArgTypes>
+	void Message_Note(const FString& Message, ArgTypes... Args)
+	{
+		UBlueprint* OwningBP = GetBlueprint();
+		if (OwningBP)
+		{
+			OwningBP->Message_Note(Message, Forward<ArgTypes>(Args)...);
+		}
+		else
+		{
+			UE_LOG(LogBlueprint, Log, TEXT("%s"), *Message);
+		}
+	}
 
+	template<typename... ArgTypes>
+	void Message_Warn(const FString& Message, ArgTypes... Args)
+	{
+		UBlueprint* OwningBP = GetBlueprint();
+		if (OwningBP)
+		{
+			OwningBP->Message_Warn(Message, Forward<ArgTypes>(Args)...);
+		}
+		else
+		{
+			UE_LOG(LogBlueprint, Warning, TEXT("%s"), *Message);
+		}
+	}
+
+	template<typename... ArgTypes>
+	void Message_Error(const FString& Message, ArgTypes... Args)
+	{
+		UBlueprint* OwningBP = GetBlueprint();
+		if (OwningBP)
+		{
+			OwningBP->Message_Error(Message, Forward<ArgTypes>(Args)...);
+		}
+		else
+		{
+			UE_LOG(LogBlueprint, Error, TEXT("%s"), *Message);
+		}
+	}
 
 	friend class FKismetCompilerContext;
 
diff --git a/Engine/Source/Editor/BlueprintGraph/Private/BlueprintTypePromotion.cpp b/Engine/Source/Editor/BlueprintGraph/Private/BlueprintTypePromotion.cpp
--- a/Engine/Source/Editor/BlueprintGraph/Private/BlueprintTypePromotion.cpp
+++ b/Engine/Source/Editor/BlueprintGraph/Private/BlueprintTypePromotion.cpp
@@ -125,11 +125,7 @@ bool FTypePromotion::IsValidPromotion(const FEdGraphPinType& A, const FEdGraphPi
 	{
 		const UEdGraphSchema_K2* K2Schema = GetDefault<UEdGraphSchema_K2>();
 
-		FName DummyName;
-		UClass* DummyClass = nullptr;
-		UK2Node* DummyNode = nullptr;
-
-		return K2Schema->SearchForAutocastFunction(A, B, /*out*/ DummyName, DummyClass);
+		return K2Schema->SearchForAutocastFunction(A, B).IsSet();
 	}
 	else
 	{
@@ -139,14 +135,13 @@ bool FTypePromotion::IsValidPromotion(const FEdGraphPinType& A, const FEdGraphPi
 
 bool FTypePromotion::HasStructConversion(const UEdGraphPin* InputPin, const UEdGraphPin* OutputPin)
 {
-	const UEdGraphSchema_K2* K2Schema = GetDefault<UEdGraphSchema_K2>();
+	check(InputPin);
+	check(OutputPin);
 
-	FName DummyName;
-	UClass* DummyClass = nullptr;
-	UK2Node* DummyNode = nullptr;
+	const UEdGraphSchema_K2* K2Schema = GetDefault<UEdGraphSchema_K2>();
 
-	const bool bCanAutocast = K2Schema->SearchForAutocastFunction(OutputPin->PinType, InputPin->PinType, /*out*/ DummyName, DummyClass);
-	const bool bCanAutoConvert = K2Schema->FindSpecializedConversionNode(OutputPin, InputPin, false, /* out */ DummyNode);
+	const bool bCanAutocast = K2Schema->SearchForAutocastFunction(OutputPin->PinType, InputPin->PinType).IsSet();
+	const bool bCanAutoConvert = K2Schema->FindSpecializedConversionNode(OutputPin->PinType, *InputPin, false).IsSet();
 	
 	return bCanAutocast || bCanAutoConvert;
 }
diff --git a/Engine/Source/Editor/BlueprintGraph/Private/EdGraphSchema_K2.cpp b/Engine/Source/Editor/BlueprintGraph/Private/EdGraphSchema_K2.cpp
--- a/Engine/Source/Editor/BlueprintGraph/Private/EdGraphSchema_K2.cpp
+++ b/Engine/Source/Editor/BlueprintGraph/Private/EdGraphSchema_K2.cpp
@@ -2167,6 +2167,9 @@ static FText GetPinIncompatibilityReason(const UEdGraphPin* PinA, const UEdGraph
 
 const FPinConnectionResponse UEdGraphSchema_K2::CanCreateConnection(const UEdGraphPin* PinA, const UEdGraphPin* PinB) const
 {
+	check(PinA);
+	check(PinB);
+
 	const UK2Node* OwningNodeA = Cast<UK2Node>(PinA->GetOwningNodeUnchecked());
 	const UK2Node* OwningNodeB = Cast<UK2Node>(PinB->GetOwningNodeUnchecked());
 
@@ -2200,14 +2203,17 @@ const FPinConnectionResponse UEdGraphSchema_K2::CanCreateConnection(const UEdGra
 	}
 
 	// Compare the directions
-	const UEdGraphPin* InputPin = NULL;
-	const UEdGraphPin* OutputPin = NULL;
+	const UEdGraphPin* InputPin = nullptr;
+	const UEdGraphPin* OutputPin = nullptr;
 
 	if (!CategorizePinsByDirection(PinA, PinB, /*out*/ InputPin, /*out*/ OutputPin))
 	{
 		return FPinConnectionResponse(CONNECT_RESPONSE_DISALLOW, TEXT("Directions are not compatible"));
 	}
 
+	check(InputPin);
+	check(OutputPin);
+
 	bool bIgnoreArray = false;
 	if (const UK2Node* OwningNode = Cast<UK2Node>(InputPin->GetOwningNode()))
 	{
@@ -2258,12 +2264,8 @@ const FPinConnectionResponse UEdGraphSchema_K2::CanCreateConnection(const UEdGra
 		}
 
 		// Autocasting
-		FName DummyName;
-		UClass* DummyClass;
-		UK2Node* DummyNode;
-
-		const bool bCanAutocast = SearchForAutocastFunction(OutputPin->PinType, InputPin->PinType, /*out*/ DummyName, DummyClass);
-		const bool bCanAutoConvert = FindSpecializedConversionNode(OutputPin, InputPin, false, /* out */ DummyNode);
+		const bool bCanAutocast = SearchForAutocastFunction(OutputPin->PinType, InputPin->PinType).IsSet();
+		const bool bCanAutoConvert = FindSpecializedConversionNode(OutputPin->PinType, *InputPin, false).IsSet();
 
 		if (bCanAutocast || bCanAutoConvert)
 		{
@@ -2482,25 +2484,36 @@ void UEdGraphSchema_K2::Shutdown()
 
 
 bool UEdGraphSchema_K2::SearchForAutocastFunction(const FEdGraphPinType& OutputPinType, const FEdGraphPinType& InputPinType, /*out*/ FName& TargetFunction, /*out*/ UClass*& FunctionOwner) const
+{
+	TOptional<FSearchForAutocastFunctionResults> Result = SearchForAutocastFunction(OutputPinType, InputPinType);
+	if (Result)
+	{
+		TargetFunction = Result->TargetFunction;
+		FunctionOwner = Result->FunctionOwner;
+		return true;
+	}
+
+	return false;
+}
+
+TOptional<UEdGraphSchema_K2::FSearchForAutocastFunctionResults> UEdGraphSchema_K2::SearchForAutocastFunction(const FEdGraphPinType& OutputPinType, const FEdGraphPinType& InputPinType) const
 {
 	TRACE_CPUPROFILER_EVENT_SCOPE(WILD_UEdGraphSchema_K2::SearchForAutocastFunction);
+
 	// NOTE: Under no circumstances should anyone *ever* add a questionable cast to this function.
 	// If it could be at all confusing why a function is provided, to even a novice user, err on the side of do not cast!!!
 	// This includes things like string->int (does it do length, atoi, or what?) that would be autocasts in a traditional scripting language
 
-	TargetFunction = NAME_None;
-	FunctionOwner = nullptr;
+	TOptional<FSearchForAutocastFunctionResults> Result;
 
 	if (OutputPinType.ContainerType != InputPinType.ContainerType)
 	{
 		if (OutputPinType.IsSet() && InputPinType.IsArray())
 		{
-			UFunction* Function = UBlueprintSetLibrary::StaticClass()->FindFunctionByName(GET_MEMBER_NAME_CHECKED(UBlueprintSetLibrary, Set_ToArray));
-			TargetFunction = Function->GetFName();
-			FunctionOwner = Function->GetOwnerClass();
-			return true;
+			const UFunction* Function = UBlueprintSetLibrary::StaticClass()->FindFunctionByName(GET_MEMBER_NAME_CHECKED(UBlueprintSetLibrary, Set_ToArray));
+			Result = { Function->GetFName(), Function->GetOwnerClass() };
 		}
-		
+
 		// Skip the other special cases if container check fails, but allow checking the autocast map
 	}
 	else
@@ -2508,14 +2521,13 @@ bool UEdGraphSchema_K2::SearchForAutocastFunction(const FEdGraphPinType& OutputP
 		// SPECIAL CASES, not supported by FAutocastFunctionMap.
 		if ((OutputPinType.PinCategory == PC_Interface) && (InputPinType.PinCategory == PC_Object))
 		{
-			UClass const* InputClass = Cast<UClass const>(InputPinType.PinSubCategoryObject.Get());
+			const UClass* InputClass = Cast<const UClass>(InputPinType.PinSubCategoryObject.Get());
 
-			bool const bInputIsUObject = (InputClass && (InputClass == UObject::StaticClass()));
+			const bool bInputIsUObject = (InputClass && (InputClass == UObject::StaticClass()));
 			if (bInputIsUObject)
 			{
 				UFunction* Function = UKismetSystemLibrary::StaticClass()->FindFunctionByName(GET_MEMBER_NAME_CHECKED(UKismetSystemLibrary, Conv_InterfaceToObject));
-				TargetFunction = Function->GetFName();
-				FunctionOwner = Function->GetOwnerClass();
+				Result = { Function->GetFName(), Function->GetOwnerClass() };
 			}
 		}
 		else if (OutputPinType.PinCategory == PC_Object)
@@ -2529,15 +2541,13 @@ bool UEdGraphSchema_K2::SearchForAutocastFunction(const FEdGraphPinType& OutputP
 					OutputClass->IsChildOf(InputClass))
 				{
 					UFunction* Function = UGameplayStatics::StaticClass()->FindFunctionByName(GET_MEMBER_NAME_CHECKED(UGameplayStatics, GetObjectClass));
-					TargetFunction = Function->GetFName();
-					FunctionOwner = Function->GetOwnerClass();
+					Result = { Function->GetFName(), Function->GetOwnerClass() };
 				}
 			}
 			else if (InputPinType.PinCategory == PC_String)
 			{
 				UFunction* Function = UKismetSystemLibrary::StaticClass()->FindFunctionByName(GET_MEMBER_NAME_CHECKED(UKismetSystemLibrary, GetDisplayName));
-				TargetFunction = Function->GetFName();
-				FunctionOwner = Function->GetOwnerClass();
+				Result = { Function->GetFName(), Function->GetOwnerClass() };
 			}
 		}
 		else if (OutputPinType.PinCategory == PC_Class)
@@ -2545,8 +2555,7 @@ bool UEdGraphSchema_K2::SearchForAutocastFunction(const FEdGraphPinType& OutputP
 			if (InputPinType.PinCategory == PC_String)
 			{
 				UFunction* Function = UKismetSystemLibrary::StaticClass()->FindFunctionByName(GET_MEMBER_NAME_CHECKED(UKismetSystemLibrary, GetClassDisplayName));
-				TargetFunction = Function->GetFName();
-				FunctionOwner = Function->GetOwnerClass();
+				Result = { Function->GetFName(), Function->GetOwnerClass() };
 			}
 		}
 		else if (OutputPinType.PinCategory == PC_Struct)
@@ -2558,23 +2567,21 @@ bool UEdGraphSchema_K2::SearchForAutocastFunction(const FEdGraphPinType& OutputP
 				if ((InputPinType.PinCategory == PC_Struct) && (InputStructType == TBaseStructure<FTransform>::Get()))
 				{
 					UFunction* Function = UKismetMathLibrary::StaticClass()->FindFunctionByName(GET_MEMBER_NAME_CHECKED(UKismetMathLibrary, MakeTransform));
-					TargetFunction = Function->GetFName();
-					FunctionOwner = Function->GetOwnerClass();
+					Result = { Function->GetFName(), Function->GetOwnerClass() };
 				}
 			}
 		}
 	}
 
 	// Try looking for a marked up autocast if we've not found a built-in one that works
-	if (TargetFunction == NAME_None)
+	if (!Result.IsSet())
 	{
-		auto FindAndSetCastFunction = [&TargetFunction, &FunctionOwner](const FEdGraphPinType& OutputPinType, const FEdGraphPinType& InputPinType)
+		auto FindAndSetCastFunction = [&Result](const FEdGraphPinType& OutputPinType, const FEdGraphPinType& InputPinType)
 		{
 			const FAutocastFunctionMap& AutocastFunctionMap = FAutocastFunctionMap::Get();
-			if (const UFunction* Func = AutocastFunctionMap.Find(OutputPinType, InputPinType))
+			if (const UFunction* Function = AutocastFunctionMap.Find(OutputPinType, InputPinType))
 			{
-				TargetFunction = Func->GetFName();
-				FunctionOwner = Func->GetOwnerClass();
+				Result = { Function->GetFName(), Function->GetOwnerClass() };
 				return true;
 			}
 			return false;
@@ -2600,41 +2607,65 @@ bool UEdGraphSchema_K2::SearchForAutocastFunction(const FEdGraphPinType& OutputP
 		}
 	}
 
-	return TargetFunction != NAME_None;
+	return Result;
 }
 
+
 bool UEdGraphSchema_K2::FindSpecializedConversionNode(const UEdGraphPin* OutputPin, const UEdGraphPin* InputPin, bool bCreateNode, /*out*/ UK2Node*& TargetNode) const
 {
-	return FindSpecializedConversionNode(OutputPin->PinType, InputPin, bCreateNode, TargetNode);
+	TOptional<FindSpecializedConversionNodeResults> Result = FindSpecializedConversionNode(OutputPin->PinType, *InputPin, bCreateNode);
+	if (Result)
+	{
+		TargetNode = Result->TargetNode;
+		return true;
+	}
+
+	return false;
 }
 
 bool UEdGraphSchema_K2::FindSpecializedConversionNode(const FEdGraphPinType& OutputPinType, const UEdGraphPin* InputPin, bool bCreateNode, UK2Node*& TargetNode) const
 {
-	bool bCanConvert = false;
-	TargetNode = nullptr;
-	if (!InputPin)
+	TOptional<FindSpecializedConversionNodeResults> Result = FindSpecializedConversionNode(OutputPinType, *InputPin, bCreateNode);
+	if (Result)
 	{
-		return bCanConvert;
+		TargetNode = Result->TargetNode;
+		return true;
 	}
 
-	// Conversion for scalar -> array
-	if( (!OutputPinType.IsContainer() && InputPin->PinType.IsArray()) && ArePinTypesCompatible(OutputPinType, InputPin->PinType, nullptr, true))
+	return false;
+}
+
+TOptional<UEdGraphSchema_K2::FindSpecializedConversionNodeResults> UEdGraphSchema_K2::FindSpecializedConversionNode(const FEdGraphPinType& OutputPinType, const UEdGraphPin& InputPin, bool bCreateNode) const
+{
+	TOptional<UEdGraphSchema_K2::FindSpecializedConversionNodeResults> Result;
+	FEdGraphPinType InputPinType = InputPin.PinType;
+
+	const bool bConvertScalarToArray =
+		!OutputPinType.IsContainer() &&
+		InputPinType.IsArray() &&
+		ArePinTypesCompatible(OutputPinType, InputPinType, nullptr, true);
+
+	const bool bTryAlternateObjectProperty =
+		InputPin.GetOwningNode()->IsA(UK2Node_CallFunction::StaticClass()) &&
+		IsSelfPin(InputPin) &&
+		((OutputPinType.PinCategory == PC_Object) || (OutputPinType.PinCategory == PC_Interface));
+
+	if (bConvertScalarToArray)
 	{
-		bCanConvert = true;
-		if(bCreateNode)
+		Result = { nullptr };
+		if (bCreateNode)
 		{
-			TargetNode = NewObject<UK2Node_MakeArray>();
+			Result->TargetNode = NewObject<UK2Node_MakeArray>();
 		}
 	}
 	// If connecting an object to a 'call function' self pin, and not currently compatible, see if there is a property we can call a function on
-	else if (InputPin->GetOwningNode()->IsA(UK2Node_CallFunction::StaticClass()) && IsSelfPin(*InputPin) && 
-		((OutputPinType.PinCategory == PC_Object) || (OutputPinType.PinCategory == PC_Interface)))
+	else if (bTryAlternateObjectProperty)
 	{
-		UK2Node_CallFunction* CallFunctionNode = (UK2Node_CallFunction*)(InputPin->GetOwningNode());
-		UClass* OutputPinClass = Cast<UClass>(OutputPinType.PinSubCategoryObject.Get());
+		const UK2Node_CallFunction* CallFunctionNode = CastChecked<UK2Node_CallFunction>(InputPin.GetOwningNode());
+		const UClass* OutputPinClass = Cast<UClass>(OutputPinType.PinSubCategoryObject.Get());
+		const UClass* FunctionClass = CallFunctionNode->FunctionReference.GetMemberParentClass(CallFunctionNode->GetBlueprintClassFromNode());
 
-		UClass* FunctionClass = CallFunctionNode->FunctionReference.GetMemberParentClass(CallFunctionNode->GetBlueprintClassFromNode());
-		if(FunctionClass != NULL && OutputPinClass != NULL)
+		if(FunctionClass != nullptr && OutputPinClass != nullptr)
 		{
 			// Iterate over object properties..
 			for (TFieldIterator<FObjectProperty> PropIt(OutputPinClass, EFieldIteratorFlags::IncludeSuper); PropIt; ++PropIt)
@@ -2643,14 +2674,12 @@ bool UEdGraphSchema_K2::FindSpecializedConversionNode(const FEdGraphPinType& Out
 				// .. if we have a blueprint visible var, and is of the type which contains this function..
 				if(ObjProp->HasAllPropertyFlags(CPF_BlueprintVisible) && ObjProp->PropertyClass && ObjProp->PropertyClass->IsChildOf(FunctionClass))
 				{
-					// say we can convert
-					bCanConvert = true;
-					// Create 'get variable' node
-					if(bCreateNode)
+					Result = { nullptr };
+					if (bCreateNode)
 					{
 						UK2Node_VariableGet* GetNode = NewObject<UK2Node_VariableGet>();
 						GetNode->VariableReference.SetFromField<FProperty>(ObjProp, false);
-						TargetNode = GetNode;
+						Result->TargetNode = GetNode;
 					}
 				}
 			}
@@ -2658,96 +2687,97 @@ bool UEdGraphSchema_K2::FindSpecializedConversionNode(const FEdGraphPinType& Out
 		}	
 	}
 
-	if(!bCanConvert)
+	if (!Result.IsSet())
 	{
 		// CHECK ENUM TO NAME CAST
-		const bool bInoputMatch = InputPin && !InputPin->PinType.IsContainer() && ((PC_Name == InputPin->PinType.PinCategory) || (PC_String == InputPin->PinType.PinCategory));
-		const bool bOutputMatch = !OutputPinType.IsContainer() && (PC_Byte == OutputPinType.PinCategory) && (NULL != Cast<UEnum>(OutputPinType.PinSubCategoryObject.Get()));
-		if(bOutputMatch && bInoputMatch)
+		const bool bInputMatch = !InputPin.PinType.IsContainer() && ((PC_Name == InputPin.PinType.PinCategory) || (PC_String == InputPin.PinType.PinCategory));
+		const bool bOutputMatch = !OutputPinType.IsContainer() && (PC_Byte == OutputPinType.PinCategory) && (nullptr != Cast<UEnum>(OutputPinType.PinSubCategoryObject.Get()));
+		if(bOutputMatch && bInputMatch)
 		{
-			bCanConvert = true;
-			if(bCreateNode)
+			Result = { nullptr };
+			if (bCreateNode)
 			{
-				check(NULL == TargetNode);
-				if(PC_Name == InputPin->PinType.PinCategory)
+				if(PC_Name == InputPin.PinType.PinCategory)
 				{
-					TargetNode = NewObject<UK2Node_GetEnumeratorName>();
+					Result->TargetNode = NewObject<UK2Node_GetEnumeratorName>();
 				}
-				else if(PC_String == InputPin->PinType.PinCategory)
+				else if(PC_String == InputPin.PinType.PinCategory)
 				{
-					TargetNode = NewObject<UK2Node_GetEnumeratorNameAsString>();
+					Result->TargetNode = NewObject<UK2Node_GetEnumeratorNameAsString>();
 				}
 			}
 		}
 	}
 
-	if (!bCanConvert && InputPin)
+	if (!Result.IsSet())
 	{
-		FEdGraphPinType const& InputType  = InputPin->PinType;
-		FEdGraphPinType const& OutputType = OutputPinType;
-
 		// CHECK BYTE TO ENUM CAST
-		UEnum* Enum = Cast<UEnum>(InputType.PinSubCategoryObject.Get());
-		const bool bInputIsEnum = !InputType.IsContainer() && (PC_Byte == InputType.PinCategory) && Enum;
-		const bool bOutputIsByte = !OutputType.IsContainer() && (PC_Byte == OutputType.PinCategory);
+		UEnum* Enum = Cast<UEnum>(InputPinType.PinSubCategoryObject.Get());
+		const bool bInputIsEnum = !InputPinType.IsContainer() && (PC_Byte == InputPinType.PinCategory) && Enum;
+		const bool bOutputIsByte = !OutputPinType.IsContainer() && (PC_Byte == OutputPinType.PinCategory);
 		if (bInputIsEnum && bOutputIsByte)
 		{
-			bCanConvert = true;
+			Result = { nullptr };
 			if(bCreateNode)
 			{
 				UK2Node_CastByteToEnum* CastByteToEnum = NewObject<UK2Node_CastByteToEnum>();
 				CastByteToEnum->Enum = Enum;
 				CastByteToEnum->bSafe = true;
-				TargetNode = CastByteToEnum;
+				Result->TargetNode = CastByteToEnum;
 			}
 		}
 		else
 		{
-			UClass* InputClass  = FBlueprintEditorUtils::GetTypeForPin(*InputPin);
-			UClass* OutputClass = Cast<UClass>(OutputPinType.PinSubCategoryObject.Get());
+			UClass* InputClass  = FBlueprintEditorUtils::GetTypeForPin(InputPin);
+			const UClass* OutputClass = Cast<UClass>(OutputPinType.PinSubCategoryObject.Get());
 
-			if ((OutputType.PinCategory == PC_Interface) && (InputType.PinCategory == PC_Object))
+			bool bNeedsDynamicCast = false;
+			if ((OutputPinType.PinCategory == PC_Interface) && (InputPinType.PinCategory == PC_Object))
 			{
-				bCanConvert = (InputClass && OutputClass) && (InputClass->ImplementsInterface(OutputClass) || OutputClass->IsChildOf(InputClass));
+				bNeedsDynamicCast = (InputClass && OutputClass) && (InputClass->ImplementsInterface(OutputClass) || OutputClass->IsChildOf(InputClass));
 			}
-			else if (OutputType.PinCategory == PC_Object)
+			else if (OutputPinType.PinCategory == PC_Object)
 			{
 				UBlueprintEditorSettings const* BlueprintSettings = GetDefault<UBlueprintEditorSettings>();
-				if ((InputType.PinCategory == PC_Object) && BlueprintSettings->bAutoCastObjectConnections)
+				if ((InputPinType.PinCategory == PC_Object) && BlueprintSettings->bAutoCastObjectConnections)
 				{
-					bCanConvert = (InputClass && OutputClass) && InputClass->IsChildOf(OutputClass);
+					bNeedsDynamicCast = (InputClass && OutputClass) && InputClass->IsChildOf(OutputClass);
 				}
 			}
 
-			if (bCanConvert && bCreateNode)
+			if (bNeedsDynamicCast)
 			{
-				UK2Node_DynamicCast* DynCastNode = NewObject<UK2Node_DynamicCast>();
-				DynCastNode->TargetType = InputClass;
-				DynCastNode->SetPurity(true);
-				TargetNode = DynCastNode;
+				Result = { nullptr };
+				if (bCreateNode)
+				{
+					UK2Node_DynamicCast* DynCastNode = NewObject<UK2Node_DynamicCast>();
+					DynCastNode->TargetType = InputClass;
+					DynCastNode->SetPurity(true);
+					Result->TargetNode = DynCastNode;
+				}
 			}
 
-			if (!bCanConvert && InputClass && OutputClass && OutputClass->IsChildOf(InputClass))
+			if (!bNeedsDynamicCast && InputClass && OutputClass && OutputClass->IsChildOf(InputClass))
 			{
-				const bool bConvertAsset = (OutputType.PinCategory == PC_SoftObject) && (InputType.PinCategory == PC_Object);
-				const bool bConvertAssetClass = (OutputType.PinCategory == PC_SoftClass) && (InputType.PinCategory == PC_Class);
-				const bool bConvertToAsset = (OutputType.PinCategory == PC_Object) && (InputType.PinCategory == PC_SoftObject);
-				const bool bConvertToAssetClass = (OutputType.PinCategory == PC_Class) && (InputType.PinCategory == PC_SoftClass);
+				const bool bConvertAsset = (OutputPinType.PinCategory == PC_SoftObject) && (InputPinType.PinCategory == PC_Object);
+				const bool bConvertAssetClass = (OutputPinType.PinCategory == PC_SoftClass) && (InputPinType.PinCategory == PC_Class);
+				const bool bConvertToAsset = (OutputPinType.PinCategory == PC_Object) && (InputPinType.PinCategory == PC_SoftObject);
+				const bool bConvertToAssetClass = (OutputPinType.PinCategory == PC_Class) && (InputPinType.PinCategory == PC_SoftClass);
 
 				if (bConvertAsset || bConvertAssetClass || bConvertToAsset || bConvertToAssetClass)
 				{
-					bCanConvert = true;
+					Result = { nullptr };
 					if (bCreateNode)
 					{
 						UK2Node_ConvertAsset* ConvertAssetNode = NewObject<UK2Node_ConvertAsset>();
-						TargetNode = ConvertAssetNode;
+						Result->TargetNode = ConvertAssetNode;
 					}
 				}
 			}
 		}
 	}
 
-	return bCanConvert;
+	return Result;
 }
 
 void UEdGraphSchema_K2::AutowireConversionNode(UEdGraphPin* InputPin, UEdGraphPin* OutputPin, UEdGraphNode* ConversionNode) const
@@ -2799,27 +2829,25 @@ bool UEdGraphSchema_K2::CreateAutomaticConversionNodeAndConnections(UEdGraphPin*
 		return false;
 	}
 
-	FName TargetFunctionName;
-	UClass* ClassContainingConversionFunction = nullptr;
-	TSubclassOf<UK2Node> ConversionNodeClass;
+	check(InputPin);
+	check(OutputPin);
 
 	UK2Node* TemplateConversionNode = nullptr;
 
-	if (SearchForAutocastFunction(OutputPin->PinType, InputPin->PinType, /*out*/ TargetFunctionName, /*out*/ClassContainingConversionFunction))
+	if (TOptional<FSearchForAutocastFunctionResults> AutocastResult = SearchForAutocastFunction(OutputPin->PinType, InputPin->PinType))
 	{
 		// Create a new call function node for the casting operator
 		UK2Node_CallFunction* TemplateNode = NewObject<UK2Node_CallFunction>();
-		TemplateNode->FunctionReference.SetExternalMember(TargetFunctionName, ClassContainingConversionFunction);
-		//TemplateNode->bIsBeadFunction = true;
+		TemplateNode->FunctionReference.SetExternalMember(AutocastResult->TargetFunction, AutocastResult->FunctionOwner);
 
 		TemplateConversionNode = TemplateNode;
 	}
-	else
+	else if (TOptional<FindSpecializedConversionNodeResults> ConversionResult = FindSpecializedConversionNode(OutputPin->PinType, *InputPin, true))
 	{
-		FindSpecializedConversionNode(OutputPin, InputPin, true, /*out*/ TemplateConversionNode);
+		TemplateConversionNode = ConversionResult->TargetNode;
 	}
 
-	if (TemplateConversionNode != NULL)
+	if (TemplateConversionNode != nullptr)
 	{
 		// Determine where to position the new node (assuming it isn't going to get beaded)
 		FVector2D AverageLocation = CalculateAveragePositionBetweenNodes(InputPin, OutputPin);
@@ -3442,7 +3470,7 @@ bool UEdGraphSchema_K2::GetPropertyCategoryInfo(const FProperty* TestProperty, F
 		// K2 only supports byte enums right now - any violations should have been caught by UHT or the editor
 		if (!EnumProperty->GetUnderlyingProperty()->IsA<FByteProperty>())
 		{
-			OutCategory = TEXT("unsupported_enum_type");
+			OutCategory = TEXT("unsupported_enum_type: enum size is larger than a byte");
 			return false;
 		}
 
diff --git a/Engine/Source/Editor/BlueprintGraph/Private/K2Node.cpp b/Engine/Source/Editor/BlueprintGraph/Private/K2Node.cpp
--- a/Engine/Source/Editor/BlueprintGraph/Private/K2Node.cpp
+++ b/Engine/Source/Editor/BlueprintGraph/Private/K2Node.cpp
@@ -1827,45 +1827,6 @@ bool UK2Node::CanCreateUnderSpecifiedSchema(const UEdGraphSchema* DesiredSchema)
 	return DesiredSchema->GetClass()->IsChildOf(UEdGraphSchema_K2::StaticClass());
 }
 
-void UK2Node::Message_Note(const FString& Message)
-{
-	UBlueprint* OwningBP = GetBlueprint();
-	if( OwningBP )
-	{
-		OwningBP->Message_Note(Message);
-	}
-	else
-	{
-		UE_LOG(LogBlueprint, Log, TEXT("%s"), *Message);
-	}
-}
-
-void UK2Node::Message_Warn(const FString& Message)
-{
-	UBlueprint* OwningBP = GetBlueprint();
-	if( OwningBP )
-	{
-		OwningBP->Message_Warn(Message);
-	}
-	else
-	{
-		UE_LOG(LogBlueprint, Warning, TEXT("%s"), *Message);
-	}
-}
-
-void UK2Node::Message_Error(const FString& Message)
-{
-	UBlueprint* OwningBP = GetBlueprint();
-	if( OwningBP )
-	{
-		OwningBP->Message_Error(Message);
-	}
-	else
-	{
-		UE_LOG(LogBlueprint, Error, TEXT("%s"), *Message);
-	}
-}
-
 FString UK2Node::GetDocumentationLink() const
 {
 	return TEXT("Shared/GraphNodes/Blueprint");
diff --git a/Engine/Source/Editor/BlueprintGraph/Private/K2Node_PromotableOperator.cpp b/Engine/Source/Editor/BlueprintGraph/Private/K2Node_PromotableOperator.cpp
--- a/Engine/Source/Editor/BlueprintGraph/Private/K2Node_PromotableOperator.cpp
+++ b/Engine/Source/Editor/BlueprintGraph/Private/K2Node_PromotableOperator.cpp
@@ -998,23 +998,21 @@ bool UK2Node_PromotableOperator::CreateIntermediateCast(UK2Node_CallFunction* So
 	}
 
 	UK2Node* TemplateConversionNode = nullptr;
-	FName TargetFunctionName;
-	UClass* ClassContainingConversionFunction = nullptr;
 	TSubclassOf<UK2Node> ConversionNodeClass;
 
-	if (Schema->SearchForAutocastFunction(InputPin->PinType, OutputPin->PinType, /*out*/ TargetFunctionName, /*out*/ClassContainingConversionFunction))
+	if (TOptional<UEdGraphSchema_K2::FSearchForAutocastFunctionResults> AutoCastResults = Schema->SearchForAutocastFunction(InputPin->PinType, OutputPin->PinType))
 	{
 		// Create a new call function node for the casting operator
 		UK2Node_CallFunction* TemplateNode = SourceGraph->CreateIntermediateNode<UK2Node_CallFunction>();
-		TemplateNode->FunctionReference.SetExternalMember(TargetFunctionName, ClassContainingConversionFunction);
+		TemplateNode->FunctionReference.SetExternalMember(AutoCastResults->TargetFunction, AutoCastResults->FunctionOwner);
 		TemplateConversionNode = TemplateNode;
 		TemplateNode->AllocateDefaultPins();
 		CompilerContext.MessageLog.NotifyIntermediateObjectCreation(TemplateNode, this);
 	}
-	else if (Schema->FindSpecializedConversionNode(InputPin, OutputPin, true, /*out*/ TemplateConversionNode))
+	else if (TOptional<UEdGraphSchema_K2::FindSpecializedConversionNodeResults> ConversionNodeResults = Schema->FindSpecializedConversionNode(InputPin->PinType, *OutputPin, true))
 	{
 		FVector2D AverageLocation = UEdGraphSchema_K2::CalculateAveragePositionBetweenNodes(InputPin, OutputPin);		
-		TemplateConversionNode = FEdGraphSchemaAction_K2NewNode::SpawnNodeFromTemplate<UK2Node>(SourceGraph, TemplateConversionNode, AverageLocation);		
+		TemplateConversionNode = FEdGraphSchemaAction_K2NewNode::SpawnNodeFromTemplate<UK2Node>(SourceGraph, ConversionNodeResults->TargetNode, AverageLocation);
 		CompilerContext.MessageLog.NotifyIntermediateObjectCreation(TemplateConversionNode, this);
 	}
 
diff --git a/Engine/Source/Editor/Kismet/Private/BPVariableDragDropAction.cpp b/Engine/Source/Editor/Kismet/Private/BPVariableDragDropAction.cpp
--- a/Engine/Source/Editor/Kismet/Private/BPVariableDragDropAction.cpp
+++ b/Engine/Source/Editor/Kismet/Private/BPVariableDragDropAction.cpp
@@ -196,19 +196,15 @@ void FKismetVariableDragDropAction::HoverTargetChanged()
 				FEdGraphPinType VariablePinType;
 				Schema->ConvertPropertyToPinType(VariableProperty, VariablePinType);
 				const bool bTypeMatch = Schema->ArePinTypesCompatible(VariablePinType, PinUnderCursor->PinType) || bIsExecPin;
-
-				FName DummyName;
-				UClass* DummyClass;
-				UK2Node* DummyNode;
-				const bool bCanAutoConvert = Schema->FindSpecializedConversionNode(VariablePinType, PinUnderCursor, false, /* out */ DummyNode);
+				const bool bCanAutoConvert = Schema->FindSpecializedConversionNode(VariablePinType, *PinUnderCursor, false).IsSet();
 				bool bCanAutocast = false;
 				if (PinUnderCursor->Direction == EGPD_Output)
 				{
-					bCanAutocast = Schema->SearchForAutocastFunction(PinUnderCursor->PinType, VariablePinType, /*out*/ DummyName, DummyClass);
+					bCanAutocast = Schema->SearchForAutocastFunction(PinUnderCursor->PinType, VariablePinType).IsSet();
 				}
 				else
 				{
-					bCanAutocast = Schema->SearchForAutocastFunction(VariablePinType, PinUnderCursor->PinType, /*out*/ DummyName, DummyClass);
+					bCanAutocast = Schema->SearchForAutocastFunction(VariablePinType, PinUnderCursor->PinType).IsSet();
 				}
 				
 				Args.Add(TEXT("PinUnderCursor"), FText::FromName(PinUnderCursor->PinName));
@@ -321,19 +317,15 @@ FReply FKismetVariableDragDropAction::DroppedOnPin(FVector2D ScreenPosition, FVe
 					FEdGraphPinType VariablePinType;
 					Schema->ConvertPropertyToPinType(VariableProperty, VariablePinType);
 					const bool bTypeMatch = Schema->ArePinTypesCompatible(VariablePinType, TargetPin->PinType) || bIsExecPin;
-					
-					FName DummyName;
-					UClass* DummyClass;
-					UK2Node* DummyNode;
-					const bool bCanAutoConvert = Schema->FindSpecializedConversionNode(VariablePinType, TargetPin, false, /* out */ DummyNode);
+					const bool bCanAutoConvert = Schema->FindSpecializedConversionNode(VariablePinType, *TargetPin, false).IsSet();
 					bool bCanAutocast = false;
 					if (TargetPin->Direction == EGPD_Output)
 					{
-						bCanAutocast = Schema->SearchForAutocastFunction(TargetPin->PinType, VariablePinType, /*out*/ DummyName, DummyClass);
+						bCanAutocast = Schema->SearchForAutocastFunction(TargetPin->PinType, VariablePinType).IsSet();
 					}
 					else
 					{
-						bCanAutocast = Schema->SearchForAutocastFunction(VariablePinType, TargetPin->PinType, /*out*/ DummyName, DummyClass);
+						bCanAutocast = Schema->SearchForAutocastFunction(VariablePinType, TargetPin->PinType).IsSet();
 					}
 					
 					if ((bTypeMatch || bCanAutocast || bCanAutoConvert) && bCanWriteIfNeeded)
diff --git a/Engine/Source/Editor/UnrealEd/Public/Kismet2/CompilerResultsLog.h b/Engine/Source/Editor/UnrealEd/Public/Kismet2/CompilerResultsLog.h
--- a/Engine/Source/Editor/UnrealEd/Public/Kismet2/CompilerResultsLog.h
+++ b/Engine/Source/Editor/UnrealEd/Public/Kismet2/CompilerResultsLog.h
@@ -79,11 +79,11 @@ class UNREALED_API FCompilerResultsLog
 		TArray< TSharedRef<FCompilerEvent> > ChildEvents;
 
 		FCompilerEvent(TSharedPtr<FCompilerEvent> InParentEventScope = nullptr)
-			:Name(TEXT(""))
-			,Counter(0)
-			,StartTime(0.0)
-			,FinishTime(0.0)
-			,ParentEventScope(InParentEventScope)
+			: Name(TEXT(""))
+			, Counter(0)
+			, StartTime(0.0)
+			, FinishTime(0.0)
+			, ParentEventScope(InParentEventScope)
 		{
 		}
 
@@ -141,7 +141,7 @@ protected:
 	TMap<TWeakObjectPtr<const UEdGraphNode>, TWeakObjectPtr<const UEdGraphNode>> IntermediateTunnelNodeToTunnelInstanceMap;
 
 	// Map to track intermediate nodes back to the source macro instance nodes
-	TMap< TWeakObjectPtr<const UEdGraphNode>, TWeakObjectPtr<UEdGraphNode> > FullMacroBacktrackMap;
+	TMap<TWeakObjectPtr<const UEdGraphNode>, TWeakObjectPtr<UEdGraphNode>> FullMacroBacktrackMap;
 
 public:
 	FCompilerResultsLog(bool bIsCompatibleWithEvents = true);
@@ -166,12 +166,12 @@ public:
 	 * Write an error in to the compiler log.
 	 * Note: @@ will be replaced by node or pin links for nodes/pins passed via varargs 
 	 */
-	template<typename... Args>
-	TSharedRef<FTokenizedMessage> Error(const TCHAR* Format, Args... args)
+	template<typename... ArgTypes>
+	TSharedRef<FTokenizedMessage> Error(const TCHAR* Format, ArgTypes... Args)
 	{
 		++NumErrors;
 		TSharedRef<FTokenizedMessage> Line = FTokenizedMessage::Create(EMessageSeverity::Error);
-		InternalLogMessage(NAME_None, Format, Line, args...);
+		InternalLogMessage(NAME_None, Format, Line, Args...);
 		return Line;
 	}
 
@@ -179,12 +179,12 @@ public:
 	 * Write a warning in to the compiler log.
 	 * Note: @@ will be replaced by node or pin links for nodes/pins passed via varargs 
 	 */
-	template<typename... Args>
-	TSharedRef<FTokenizedMessage> Warning(const TCHAR* Format, Args... args)
+	template<typename... ArgTypes>
+	TSharedRef<FTokenizedMessage> Warning(const TCHAR* Format, ArgTypes... Args)
 	{
 		++NumWarnings;
 		TSharedRef<FTokenizedMessage> Line = FTokenizedMessage::Create(EMessageSeverity::Warning);
-		InternalLogMessage(NAME_None, Format, Line, args...);
+		InternalLogMessage(NAME_None, Format, Line, Args...);
 		return Line;
 	}
 
@@ -192,8 +192,8 @@ public:
 	 * Write a warning in to the compiler log.
 	 * Note: @@ will be replaced by node or pin links for nodes/pins passed via varargs 
 	 */
-	template<typename... Args>
-	void Warning(FName ID, const TCHAR* Format, Args... args)
+	template<typename... ArgTypes>
+	void Warning(FName ID, const TCHAR* Format, ArgTypes... Args)
 	{
 		if(!IsMessageEnabled(ID))
 		{
@@ -202,7 +202,7 @@ public:
 
 		++NumWarnings;
 		TSharedRef<FTokenizedMessage> Line = FTokenizedMessage::Create(EMessageSeverity::Warning);
-		InternalLogMessage(ID, Format, Line, args...);
+		InternalLogMessage(ID, Format, Line, Args...);
 		return;
 	}
 
@@ -210,11 +210,11 @@ public:
 	 * Write a note in to the compiler log.
 	 * Note: @@ will be replaced by node or pin links for nodes/pins passed via varargs 
 	 */
-	template<typename... Args>
-	TSharedRef<FTokenizedMessage> Note(const TCHAR* Format, Args... args)
+	template<typename... ArgTypes>
+	TSharedRef<FTokenizedMessage> Note(const TCHAR* Format, ArgTypes... Args)
 	{
 		TSharedRef<FTokenizedMessage> Line = FTokenizedMessage::Create(EMessageSeverity::Info);
-		InternalLogMessage(NAME_None, Format, Line, args...);
+		InternalLogMessage(NAME_None, Format, Line, Args...);
 		return Line;
 	}
 
@@ -222,12 +222,12 @@ public:
 	 * Store a potential error for a given node in the compiler log. All messages for the node can be committed to the log later by calling CommitPotentialMessages
 	 * Note: @@ will be replaced by node or pin links for nodes/pins passed via varargs 
 	 */
-	template<typename... Args>
-	TSharedRef<FTokenizedMessage> StorePotentialError(const UEdGraphNode* Source, const TCHAR* Format, Args... args)
+	template<typename... ArgTypes>
+	TSharedRef<FTokenizedMessage> StorePotentialError(const UEdGraphNode* Source, const TCHAR* Format, ArgTypes... Args)
 	{
 		TSharedRef<FTokenizedMessage> Line = FTokenizedMessage::Create(EMessageSeverity::Error);
 		TArray<UEdGraphNode*> SourceNodes;
-		Tokenize(Format, *Line, SourceNodes, args...);
+		Tokenize(Format, *Line, SourceNodes, Args...);
 		PotentialMessages.FindOrAdd(Source).Add(Line);
 		return Line;
 	}
@@ -236,12 +236,12 @@ public:
 	 * Store a potential warning for a given node in the compiler log. All messages for the node can be committed to the log later by calling CommitPotentialMessages
 	 * Note: @@ will be replaced by node or pin links for nodes/pins passed via varargs 
 	 */
-	template<typename... Args>
-	TSharedRef<FTokenizedMessage> StorePotentialWarning(const UEdGraphNode* Source, const TCHAR* Format, Args... args)
+	template<typename... ArgTypes>
+	TSharedRef<FTokenizedMessage> StorePotentialWarning(const UEdGraphNode* Source, const TCHAR* Format, ArgTypes... Args)
 	{
 		TSharedRef<FTokenizedMessage> Line = FTokenizedMessage::Create(EMessageSeverity::Warning);
 		TArray<UEdGraphNode*> SourceNodes;
-		Tokenize(Format, *Line, SourceNodes, args...);
+		Tokenize(Format, *Line, SourceNodes, Args...);
 		PotentialMessages.FindOrAdd(Source).Add(Line);
 		return Line;
 	}
@@ -250,12 +250,12 @@ public:
 	 * Store a potential note for a given node in the compiler log. All messages for the node can be committed to the log later by calling CommitPotentialMessages
 	 * Note: @@ will be replaced by node or pin links for nodes/pins passed via varargs 
 	 */
-	template<typename... Args>
-	TSharedRef<FTokenizedMessage> StorePotentialNote(const UEdGraphNode* Source, const TCHAR* Format, Args... args)
+	template<typename... ArgTypes>
+	TSharedRef<FTokenizedMessage> StorePotentialNote(const UEdGraphNode* Source, const TCHAR* Format, ArgTypes... Args)
 	{
 		TSharedRef<FTokenizedMessage> Line = FTokenizedMessage::Create(EMessageSeverity::Info);
 		TArray<UEdGraphNode*> SourceNodes;
-		Tokenize(Format, *Line, SourceNodes, args...);
+		Tokenize(Format, *Line, SourceNodes, Args...);
 		PotentialMessages.FindOrAdd(Source).Add(Line);
 		return Line;
 	}
@@ -350,8 +350,8 @@ protected:
 		OutMessage.AddToken(FTextToken::Create(FText::FromString(Text)));
 	}
 
-	template<typename T, typename... Args>
-	void Tokenize(const TCHAR* Format, FTokenizedMessage &OutMessage, TArray<UEdGraphNode*>& OutSourceNode, T First, Args... Rest)
+	template<typename T, typename... ArgTypes>
+	void Tokenize(const TCHAR* Format, FTokenizedMessage &OutMessage, TArray<UEdGraphNode*>& OutSourceNode, T First, ArgTypes... Rest)
 	{
 		// read to next "@@":
 		if (const TCHAR* DelimiterStr = FCString::Strstr(Format, TEXT("@@")))
@@ -372,14 +372,14 @@ protected:
 		}
 	}
 
-	template<typename... Args>
-	void InternalLogMessage(FName MessageID, const TCHAR* Format, const TSharedRef<FTokenizedMessage>& Message, Args... args)
+	template<typename... ArgTypes>
+	void InternalLogMessage(FName MessageID, const TCHAR* Format, const TSharedRef<FTokenizedMessage>& Message, ArgTypes... Args)
 	{
 		// Convention for SourceNode established by the original version of the compiler results log
 		// was to annotate the error on the first node we can find. I am preserving that behavior
 		// for this type safe, variadic version:
 		TArray<UEdGraphNode*> SourceNodes;
-		Tokenize(Format, *Message, SourceNodes, args...);
+		Tokenize(Format, *Message, SourceNodes, Args...);
 		InternalLogMessage(MessageID, Message, SourceNodes);
 	}
 
diff --git a/Engine/Source/Runtime/Engine/Classes/Engine/Blueprint.h b/Engine/Source/Runtime/Engine/Classes/Engine/Blueprint.h
--- a/Engine/Source/Runtime/Engine/Classes/Engine/Blueprint.h
+++ b/Engine/Source/Runtime/Engine/Classes/Engine/Blueprint.h
@@ -17,9 +17,14 @@
 #endif
 #include "UObject/SoftObjectPath.h"
 #include "Blueprint/BlueprintSupport.h"
+
+#if WITH_EDITOR
+#include "EngineLogs.h"
+#include "Kismet2/CompilerResultsLog.h"
+#endif
+
 #include "Blueprint.generated.h"
 
-class FCompilerResultsLog;
 class ITargetPlatform;
 class UActorComponent;
 class UEdGraph;
@@ -1212,21 +1217,56 @@ private:
 
 public:
 	/** If this blueprint is currently being compiled, the CurrentMessageLog will be the log currently being used to send logs to. */
-	class FCompilerResultsLog* CurrentMessageLog;
+	FCompilerResultsLog* CurrentMessageLog;
 
 	/** Message log for storing upgrade notes that were generated within the Blueprint, will be displayed to the compiler results each compiler and will remain until saving */
-	TSharedPtr<class FCompilerResultsLog> UpgradeNotesLog;
+	TSharedPtr<FCompilerResultsLog> UpgradeNotesLog;
 
 	/** Message log for storing pre-compile errors/notes/warnings that will only last until the next Blueprint compile */
-	TSharedPtr<class FCompilerResultsLog> PreCompileLog;
+	TSharedPtr<FCompilerResultsLog> PreCompileLog;
 
 	/** 
 	 * Sends a message to the CurrentMessageLog, if there is one available.  Otherwise, defaults to logging to the normal channels.
 	 * Should use this for node and blueprint actions that happen during compilation!
 	 */
-	void Message_Note(const FString& MessageToLog);
-	void Message_Warn(const FString& MessageToLog);
-	void Message_Error(const FString& MessageToLog);
+	template<typename... ArgTypes>
+	void Message_Note(const FString& MessageToLog, ArgTypes... Args)
+	{
+		if (CurrentMessageLog)
+		{
+			CurrentMessageLog->Note(*MessageToLog, Forward<ArgTypes>(Args)...);
+		}
+		else
+		{
+			UE_LOG(LogBlueprint, Log, TEXT("[%s] %s"), *GetName(), *MessageToLog);
+		}
+	}
+
+	template<typename... ArgTypes>
+	void Message_Warn(const FString& MessageToLog, ArgTypes... Args)
+	{
+		if (CurrentMessageLog)
+		{
+			CurrentMessageLog->Warning(*MessageToLog, Forward<ArgTypes>(Args)...);
+		}
+		else
+		{
+			UE_LOG(LogBlueprint, Warning, TEXT("[%s] %s"), *GetName(), *MessageToLog);
+		}
+	}
+
+	template<typename... ArgTypes>
+	void Message_Error(const FString& MessageToLog, ArgTypes... Args)
+	{
+		if (CurrentMessageLog)
+		{
+			CurrentMessageLog->Error(*MessageToLog, Forward<ArgTypes>(Args)...);
+		}
+		else
+		{
+			UE_LOG(LogBlueprint, Error, TEXT("[%s] %s"), *GetName(), *MessageToLog);
+		}
+	}
 #endif
 	
 #if WITH_EDITORONLY_DATA
diff --git a/Engine/Source/Runtime/Engine/Private/Blueprint.cpp b/Engine/Source/Runtime/Engine/Private/Blueprint.cpp
--- a/Engine/Source/Runtime/Engine/Private/Blueprint.cpp
+++ b/Engine/Source/Runtime/Engine/Private/Blueprint.cpp
@@ -29,7 +29,6 @@
 #include "Kismet2/BlueprintEditorUtils.h"
 #include "Kismet2/ComponentEditorUtils.h"
 #include "Kismet2/KismetEditorUtilities.h"
-#include "Kismet2/CompilerResultsLog.h"
 #include "Kismet2/StructureEditorUtils.h"
 #include "WatchPointViewer.h"
 #include "FindInBlueprintManager.h"
@@ -1816,42 +1815,6 @@ void UBlueprint::GetAllGraphs(TArray<UEdGraph*>& Graphs) const
 }
 
 #if WITH_EDITOR
-void UBlueprint::Message_Note(const FString& MessageToLog)
-{
-	if( CurrentMessageLog )
-	{
-		CurrentMessageLog->Note(*MessageToLog);
-	}
-	else
-	{
-		UE_LOG(LogBlueprint, Log, TEXT("[%s] %s"), *GetName(), *MessageToLog);
-	}
-}
-
-void UBlueprint::Message_Warn(const FString& MessageToLog)
-{
-	if( CurrentMessageLog )
-	{
-		CurrentMessageLog->Warning(*MessageToLog);
-	}
-	else
-	{
-		UE_LOG(LogBlueprint, Warning, TEXT("[%s] %s"), *GetName(), *MessageToLog);
-	}
-}
-
-void UBlueprint::Message_Error(const FString& MessageToLog)
-{
-	if( CurrentMessageLog )
-	{
-		CurrentMessageLog->Error(*MessageToLog);
-	}
-	else
-	{
-		UE_LOG(LogBlueprint, Error, TEXT("[%s] %s"), *GetName(), *MessageToLog);
-	}
-}
-
 bool UBlueprint::ChangeOwnerOfTemplates()
 {
 	struct FUniqueNewNameHelper
