Fixed typos from CL 23482483:

* FFindSpecializedConversionNodeResults was missing its leading 'F'
* The deprecation comments for FindSpecializedConversionNodeResult listed the wrong function

diff --git a/Engine/Source/Editor/BlueprintGraph/Classes/EdGraphSchema_K2.h b/Engine/Source/Editor/BlueprintGraph/Classes/EdGraphSchema_K2.h
--- a/Engine/Source/Editor/BlueprintGraph/Classes/EdGraphSchema_K2.h
+++ b/Engine/Source/Editor/BlueprintGraph/Classes/EdGraphSchema_K2.h
@@ -1125,18 +1125,18 @@ public:
 	};
 	UE_NODISCARD virtual TOptional<FSearchForAutocastFunctionResults> SearchForAutocastFunction(const FEdGraphPinType& OutputPinType, const FEdGraphPinType& InputPinType) const;
 
-	UE_DEPRECATED(5.2, "Use the FSearchForAutocastFunctionResults variant.")
+	UE_DEPRECATED(5.2, "Use the FFindSpecializedConversionNodeResults variant.")
 	virtual bool FindSpecializedConversionNode(const UEdGraphPin* OutputPin, const UEdGraphPin* InputPin, bool bCreateNode, /*out*/ class UK2Node*& TargetNode) const;
 	
-	UE_DEPRECATED(5.2, "Use the FSearchForAutocastFunctionResults variant.")
+	UE_DEPRECATED(5.2, "Use the FFindSpecializedConversionNodeResults variant.")
 	virtual bool FindSpecializedConversionNode(const FEdGraphPinType& OutputPinType, const UEdGraphPin* InputPinType, bool bCreateNode, /*out*/ class UK2Node*& TargetNode) const;
 	
 	/** Find an appropriate node that can convert from one pin type to another (not a cast; e.g. "MakeLiteralArray" node) */
-	struct FindSpecializedConversionNodeResults
+	struct FFindSpecializedConversionNodeResults
 	{
 		class UK2Node* TargetNode = nullptr;
 	};
-	UE_NODISCARD virtual TOptional<FindSpecializedConversionNodeResults> FindSpecializedConversionNode(const FEdGraphPinType& OutputPinType, const UEdGraphPin& InputPin, bool bCreateNode) const;
+	UE_NODISCARD virtual TOptional<FFindSpecializedConversionNodeResults> FindSpecializedConversionNode(const FEdGraphPinType& OutputPinType, const UEdGraphPin& InputPin, bool bCreateNode) const;
 
 	/** Create menu for variable get/set nodes which refer to a variable which does not exist. */
 	void GetNonExistentVariableMenu(FToolMenuSection& Section, const UEdGraphNode* InGraphNode, UBlueprint* OwnerBlueprint) const;
diff --git a/Engine/Source/Editor/BlueprintGraph/Private/EdGraphSchema_K2.cpp b/Engine/Source/Editor/BlueprintGraph/Private/EdGraphSchema_K2.cpp
--- a/Engine/Source/Editor/BlueprintGraph/Private/EdGraphSchema_K2.cpp
+++ b/Engine/Source/Editor/BlueprintGraph/Private/EdGraphSchema_K2.cpp
@@ -2613,7 +2613,7 @@ TOptional<UEdGraphSchema_K2::FSearchForAutocastFunctionResults> UEdGraphSchema_K
 
 bool UEdGraphSchema_K2::FindSpecializedConversionNode(const UEdGraphPin* OutputPin, const UEdGraphPin* InputPin, bool bCreateNode, /*out*/ UK2Node*& TargetNode) const
 {
-	TOptional<FindSpecializedConversionNodeResults> Result = FindSpecializedConversionNode(OutputPin->PinType, *InputPin, bCreateNode);
+	TOptional<FFindSpecializedConversionNodeResults> Result = FindSpecializedConversionNode(OutputPin->PinType, *InputPin, bCreateNode);
 	if (Result)
 	{
 		TargetNode = Result->TargetNode;
@@ -2625,7 +2625,7 @@ bool UEdGraphSchema_K2::FindSpecializedConversionNode(const UEdGraphPin* OutputP
 
 bool UEdGraphSchema_K2::FindSpecializedConversionNode(const FEdGraphPinType& OutputPinType, const UEdGraphPin* InputPin, bool bCreateNode, UK2Node*& TargetNode) const
 {
-	TOptional<FindSpecializedConversionNodeResults> Result = FindSpecializedConversionNode(OutputPinType, *InputPin, bCreateNode);
+	TOptional<FFindSpecializedConversionNodeResults> Result = FindSpecializedConversionNode(OutputPinType, *InputPin, bCreateNode);
 	if (Result)
 	{
 		TargetNode = Result->TargetNode;
@@ -2635,9 +2635,9 @@ bool UEdGraphSchema_K2::FindSpecializedConversionNode(const FEdGraphPinType& Out
 	return false;
 }
 
-TOptional<UEdGraphSchema_K2::FindSpecializedConversionNodeResults> UEdGraphSchema_K2::FindSpecializedConversionNode(const FEdGraphPinType& OutputPinType, const UEdGraphPin& InputPin, bool bCreateNode) const
+TOptional<UEdGraphSchema_K2::FFindSpecializedConversionNodeResults> UEdGraphSchema_K2::FindSpecializedConversionNode(const FEdGraphPinType& OutputPinType, const UEdGraphPin& InputPin, bool bCreateNode) const
 {
-	TOptional<UEdGraphSchema_K2::FindSpecializedConversionNodeResults> Result;
+	TOptional<UEdGraphSchema_K2::FFindSpecializedConversionNodeResults> Result;
 	FEdGraphPinType InputPinType = InputPin.PinType;
 
 	const bool bConvertScalarToArray =
@@ -2842,7 +2842,7 @@ bool UEdGraphSchema_K2::CreateAutomaticConversionNodeAndConnections(UEdGraphPin*
 
 		TemplateConversionNode = TemplateNode;
 	}
-	else if (TOptional<FindSpecializedConversionNodeResults> ConversionResult = FindSpecializedConversionNode(OutputPin->PinType, *InputPin, true))
+	else if (TOptional<FFindSpecializedConversionNodeResults> ConversionResult = FindSpecializedConversionNode(OutputPin->PinType, *InputPin, true))
 	{
 		TemplateConversionNode = ConversionResult->TargetNode;
 	}
diff --git a/Engine/Source/Editor/BlueprintGraph/Private/K2Node_PromotableOperator.cpp b/Engine/Source/Editor/BlueprintGraph/Private/K2Node_PromotableOperator.cpp
--- a/Engine/Source/Editor/BlueprintGraph/Private/K2Node_PromotableOperator.cpp
+++ b/Engine/Source/Editor/BlueprintGraph/Private/K2Node_PromotableOperator.cpp
@@ -1009,7 +1009,7 @@ bool UK2Node_PromotableOperator::CreateIntermediateCast(UK2Node_CallFunction* So
 		TemplateNode->AllocateDefaultPins();
 		CompilerContext.MessageLog.NotifyIntermediateObjectCreation(TemplateNode, this);
 	}
-	else if (TOptional<UEdGraphSchema_K2::FindSpecializedConversionNodeResults> ConversionNodeResults = Schema->FindSpecializedConversionNode(InputPin->PinType, *OutputPin, true))
+	else if (TOptional<UEdGraphSchema_K2::FFindSpecializedConversionNodeResults> ConversionNodeResults = Schema->FindSpecializedConversionNode(InputPin->PinType, *OutputPin, true))
 	{
 		FVector2D AverageLocation = UEdGraphSchema_K2::CalculateAveragePositionBetweenNodes(InputPin, OutputPin);		
 		TemplateConversionNode = FEdGraphSchemaAction_K2NewNode::SpawnNodeFromTemplate<UK2Node>(SourceGraph, ConversionNodeResults->TargetNode, AverageLocation);
