Updated the CreateOpTable to only be refreshed when a dll is loaded or unloaded and not when a directory is added.

Before: 2s60ms
After: 1s

diff --git a/Engine/Source/Editor/BlueprintGraph/Private/BlueprintTypePromotion.cpp b/Engine/Source/Editor/BlueprintGraph/Private/BlueprintTypePromotion.cpp
--- a/Engine/Source/Editor/BlueprintGraph/Private/BlueprintTypePromotion.cpp
+++ b/Engine/Source/Editor/BlueprintGraph/Private/BlueprintTypePromotion.cpp
@@ -84,7 +84,10 @@ void FTypePromotion::OnModulesChanged(FName ModuleThatChanged, EModuleChangeReas
 	// Any time a module is changed, there could possibly be new UFunctions that we 
 	// need to process, so we need to recreate the op table and clear the node spawners
 	// that we are using in order to avoid invalid duplicates in the graph action menu
-	FTypePromotion::RefreshPromotionTables();
+	if (ReasonForChange == EModuleChangeReason::ModuleLoaded || ReasonForChange == EModuleChangeReason::ModuleUnloaded)
+	{
+		FTypePromotion::RefreshPromotionTables();
+	}
 }
 
 void FTypePromotion::CreatePromotionTable()
