Avoid using fallback struct for wildcard pins, allowing blueprint macro expansion to occur as normal

diff --git a/Engine/Source/Editor/BlueprintGraph/Private/K2Node_GetDataTableRow.cpp b/Engine/Source/Editor/BlueprintGraph/Private/K2Node_GetDataTableRow.cpp
--- a/Engine/Source/Editor/BlueprintGraph/Private/K2Node_GetDataTableRow.cpp
+++ b/Engine/Source/Editor/BlueprintGraph/Private/K2Node_GetDataTableRow.cpp
@@ -155,7 +155,7 @@ UScriptStruct* UK2Node_GetDataTableRow::GetDataTableRowStructType() const
 		if (ResultPin && ResultPin->LinkedTo.Num() > 0)
 		{
 			RowStructType = Cast<UScriptStruct>(ResultPin->LinkedTo[0]->PinType.PinSubCategoryObject.Get());
-			if (RowStructType == nullptr)
+			if (RowStructType == nullptr && ResultPin->PinType.PinCategory != UEdGraphSchema_K2::PC_Wildcard)
 			{
 				RowStructType = GetFallbackStruct();
 			}
