Fix SA error

diff --git a/Engine/Source/Editor/BlueprintGraph/Private/K2Node_GetDataTableRow.cpp b/Engine/Source/Editor/BlueprintGraph/Private/K2Node_GetDataTableRow.cpp
--- a/Engine/Source/Editor/BlueprintGraph/Private/K2Node_GetDataTableRow.cpp
+++ b/Engine/Source/Editor/BlueprintGraph/Private/K2Node_GetDataTableRow.cpp
@@ -164,7 +164,7 @@ UScriptStruct* UK2Node_GetDataTableRow::GetDataTableRowStructType() const
 				UEdGraphPin* Link = ResultPin->LinkedTo[LinkIndex];
 				UScriptStruct* LinkType = Cast<UScriptStruct>(Link->PinType.PinSubCategoryObject.Get());
 
-				if (RowStructType->IsChildOf(LinkType))
+				if (RowStructType && RowStructType->IsChildOf(LinkType))
 				{
 					RowStructType = LinkType;
 				}
