Make BlueprintTypePromotion thread-safe for async loading in editor
  - Scenario is: A module gets loaded and triggers a RefreshPromotionTables on the game-thread while the ALT reads FindBestMatchingFunc from the table.

#rnx

diff --git a/Engine/Source/Editor/BlueprintGraph/Classes/BlueprintTypePromotion.h b/Engine/Source/Editor/BlueprintGraph/Classes/BlueprintTypePromotion.h
--- a/Engine/Source/Editor/BlueprintGraph/Classes/BlueprintTypePromotion.h
+++ b/Engine/Source/Editor/BlueprintGraph/Classes/BlueprintTypePromotion.h
@@ -177,7 +177,7 @@ private:
 	void CreateOpTable();
 
 	/** Creates the table of what types can be promoted to others */
-	void CreatePromotionTable();
+	TMap<FName, TArray<FName>> CreatePromotionTable();
 
 	void AddOpFunction(FName OpName, UFunction* Function);
 
@@ -190,7 +190,7 @@ private:
 	FDelegateHandle OnReloadCompleteDelegateHandle;
 
 	/** A map of 'Type' to its 'available promotions'. See ctor for creation */
-	TMap<FName, TArray<FName>> PromotionTable;
+	const TMap<FName, TArray<FName>> PromotionTable;
 
 	/**
 	 * A single operator can have multiple functions associated with it; usually
@@ -199,6 +199,11 @@ private:
 	 */
 	typedef TArray<UFunction*> FFunctionsList;
 
+	/**
+	 * Protects internal data from multi-threaded access.
+	 */
+	mutable FCriticalSection Lock;
+
 	/**
 	 * A lookup table, mapping operator strings (like "Add", "Multiply", etc.) to a list
 	 * of associated functions.
diff --git a/Engine/Source/Editor/BlueprintGraph/Private/BlueprintTypePromotion.cpp b/Engine/Source/Editor/BlueprintGraph/Private/BlueprintTypePromotion.cpp
--- a/Engine/Source/Editor/BlueprintGraph/Private/BlueprintTypePromotion.cpp
+++ b/Engine/Source/Editor/BlueprintGraph/Private/BlueprintTypePromotion.cpp
@@ -12,6 +12,7 @@
 #include "Internationalization/Text.h"
 #include "Kismet/BlueprintFunctionLibrary.h"
 #include "Misc/AssertionMacros.h"
+#include "Misc/MTAccessDetector.h"
 #include "Modules/ModuleManager.h"
 #include "ProfilingDebugging/CpuProfilerTrace.h"
 #include "UObject/Class.h"
@@ -26,6 +27,11 @@
 class UBlueprintFunctionNodeSpawner;
 class UK2Node;
 
+namespace TypePromotionImpl
+{
+	UE_MT_DECLARE_RW_ACCESS_DETECTOR(InstanceAccessProtector);
+}
+
 FTypePromotion* FTypePromotion::Instance = nullptr;
 
 #define LOCTEXT_NAMESPACE "TypePromotion"
@@ -51,6 +57,7 @@ FTypePromotion& FTypePromotion::Get()
 {
 	if (Instance == nullptr)
 	{
+		UE_MT_SCOPED_WRITE_ACCESS(TypePromotionImpl::InstanceAccessProtector);
 		Instance = new FTypePromotion();
 	}
 	return *Instance;
@@ -60,14 +67,15 @@ void FTypePromotion::Shutdown()
 {
 	if (Instance)
 	{
+		UE_MT_SCOPED_WRITE_ACCESS(TypePromotionImpl::InstanceAccessProtector);
 		delete Instance;
 		Instance = nullptr;
 	}
 }
 
 FTypePromotion::FTypePromotion()
+	: PromotionTable(CreatePromotionTable())
 {
-	CreatePromotionTable();
 	CreateOpTable();
 	OnModulesChangedDelegateHandle = FModuleManager::Get().OnModulesChanged().AddStatic(&FTypePromotion::OnModulesChanged);
 	OnReloadCompleteDelegateHandle = FCoreUObjectDelegates::ReloadCompleteDelegate.AddStatic(&FTypePromotion::RefreshPromotionTables);
@@ -90,11 +98,10 @@ void FTypePromotion::OnModulesChanged(FName ModuleThatChanged, EModuleChangeReas
 	}
 }
 
-void FTypePromotion::CreatePromotionTable()
+TMap<FName, TArray<FName>> FTypePromotion::CreatePromotionTable()
 {
-	PromotionTable =
+	return
 	{
-
 		// Type_X...						Can be promoted to...
 		{ UEdGraphSchema_K2::PC_Int,		{ UEdGraphSchema_K2::PC_Real, UEdGraphSchema_K2::PC_Int64 } },
 		{ UEdGraphSchema_K2::PC_Byte,		{ UEdGraphSchema_K2::PC_Real, UEdGraphSchema_K2::PC_Int, UEdGraphSchema_K2::PC_Int64 } },
@@ -202,7 +209,8 @@ bool FTypePromotion::IsFunctionPromotionReady(const UFunction* const FuncToConsi
 bool FTypePromotion::IsFunctionPromotionReady_Internal(const UFunction* const FuncToConsider) const
 {
 	const FName FuncOpName = GetOpNameFromFunction(FuncToConsider);
-	
+
+	FScopeLock ScopeLock(&Lock);
 	if(const FFunctionsList* FuncList = OperatorTable.Find(FuncOpName))
 	{
 		return FuncList->Contains(FuncToConsider);
@@ -266,6 +274,8 @@ UFunction* FTypePromotion::FindBestMatchingFunc_Internal(FName Operation, const
 {
 	TRACE_CPUPROFILER_EVENT_SCOPE(FTypePromotionTable::FindBestMatchingFunc_Internal);
 
+	FScopeLock ScopeLock(&Lock);
+
 	const FFunctionsList* FuncList = OperatorTable.Find(Operation);
 	if (!FuncList)
 	{
@@ -471,6 +481,7 @@ void FTypePromotion::GetAllFuncsForOp_Internal(FName Operation, TArray<UFunction
 {
 	OutFuncs.Empty();
 
+	FScopeLock ScopeLock(&Lock);
 	OutFuncs.Append(OperatorTable[Operation]);
 }
 
@@ -501,14 +512,14 @@ void FTypePromotion::CreateOpTable()
 {
 	TRACE_CPUPROFILER_EVENT_SCOPE(FTypePromotion::CreateOpTable);
 	const UEdGraphSchema_K2* Schema = GetDefault<UEdGraphSchema_K2>();
-	FCriticalSection AddOpLock;
 
-	OperatorTable.Empty();
+	FCriticalSection NewOperatorTableLock;
+	TMap<FName, FFunctionsList> NewOperatorTable;
 
 	TArray<UClass*> Libraries;
 	GetDerivedClasses(UBlueprintFunctionLibrary::StaticClass(), Libraries);
 	ParallelFor(Libraries.Num(), 
-		[this, &Schema, &AddOpLock, &Libraries](int32 Index)
+		[this, &Schema, &NewOperatorTable, &NewOperatorTableLock, &Libraries](int32 Index)
 		{
 			UClass* Library = Libraries[Index];
 			// Ignore abstract libraries/classes
@@ -529,15 +540,19 @@ void FTypePromotion::CreateOpTable()
 
 				if (OpName != OperatorNames::NoOp && Schema->ConvertPropertyToPinType(Function->GetReturnProperty(), /* out */ FuncPinType))
 				{
-					FScopeLock ScopeLock(&AddOpLock);
-					AddOpFunction(OpName, Function);
+					FScopeLock ScopeLock(&NewOperatorTableLock);
+					NewOperatorTable.FindOrAdd(OpName).Add(Function);
 				}
 			}
 		});
+
+	FScopeLock ScopeLock(&Lock);
+	OperatorTable = MoveTemp(NewOperatorTable);
 }
 
 void FTypePromotion::AddOpFunction(FName OpName, UFunction* Function)
 {
+	FScopeLock ScopeLock(&Lock);
 	OperatorTable.FindOrAdd(OpName).Add(Function);
 }
 
@@ -588,17 +603,25 @@ bool FTypePromotion::IsOperatorSpawnerRegistered(UFunction const* const Func)
 
 void FTypePromotion::RegisterOperatorSpawner(FName OpName, UBlueprintFunctionNodeSpawner* Spawner)
 {
-	if(Instance && !Instance->OperatorNodeSpawnerMap.Contains(OpName) && OpName != OperatorNames::NoOp)
+	if(Instance)
 	{
-		Instance->OperatorNodeSpawnerMap.Add(OpName, Spawner);
+		FScopeLock ScopeLock(&Instance->Lock);
+		if (OpName != OperatorNames::NoOp && !Instance->OperatorNodeSpawnerMap.Contains(OpName))
+		{
+			Instance->OperatorNodeSpawnerMap.Add(OpName, Spawner);
+		}
 	}
 }
 
 UBlueprintFunctionNodeSpawner* FTypePromotion::GetOperatorSpawner(FName OpName)
 {
-	if(Instance && Instance->OperatorNodeSpawnerMap.Contains(OpName))
+	if(Instance)
 	{
-		return Instance->OperatorNodeSpawnerMap[OpName];
+		FScopeLock ScopeLock(&Instance->Lock);
+		if (Instance->OperatorNodeSpawnerMap.Contains(OpName))
+		{
+			return Instance->OperatorNodeSpawnerMap[OpName];
+		}
 	}
 
 	return nullptr;
@@ -608,6 +631,7 @@ void FTypePromotion::ClearNodeSpawners()
 {
 	if (Instance)
 	{
+		FScopeLock ScopeLock(&Instance->Lock);
 		Instance->OperatorNodeSpawnerMap.Empty();
 	}
 }
@@ -615,10 +639,15 @@ void FTypePromotion::ClearNodeSpawners()
 void FTypePromotion::RefreshPromotionTables(EReloadCompleteReason Reason /* = EReloadCompleteReason::None */)
 {
 	TRACE_CPUPROFILER_EVENT_SCOPE(FTypePromotionTable::RefreshPromotionTables);
-	
-	FTypePromotion::ClearNodeSpawners();
+
 	if(Instance)
 	{
+		// Hold the lock during the whole operation since we don't want readers to see
+		// any intermediate state.
+		FScopeLock ScopeLock(&Instance->Lock);
+
+		FTypePromotion::ClearNodeSpawners();
+
 		Instance->OperatorTable.Empty();
 		Instance->CreateOpTable();
 	}
