Add UPARAM(Required) markup for marking function parameters as required (must be linked to some other node in a Blueprint) - example usage:

UFUNCTION(BlueprintCallable)
static void TestRequiredIntParameter(UPARAM(Required) int32 Value);

Useful for marking pointer params as required so they are not null by default, but also applicable to complex types with custom make nodes

diff --git a/Engine/Source/Editor/BlueprintGraph/Classes/K2Node_CallFunction.h b/Engine/Source/Editor/BlueprintGraph/Classes/K2Node_CallFunction.h
--- a/Engine/Source/Editor/BlueprintGraph/Classes/K2Node_CallFunction.h
+++ b/Engine/Source/Editor/BlueprintGraph/Classes/K2Node_CallFunction.h
@@ -126,6 +126,7 @@ public:
 	virtual bool ShouldDrawAsBead() const override;
 	virtual FText GetCompactNodeTitle() const override;
 	virtual void PostPasteNode() override;
+	virtual bool CanSplitPin(const UEdGraphPin* Pin) const override;
 	virtual void ValidateNodeDuringCompilation(class FCompilerResultsLog& MessageLog) const override;
 	virtual bool ShouldShowNodeProperties() const override;
 	virtual void GetRedirectPinNames(const UEdGraphPin& Pin, TArray<FString>& RedirectPinNames) const override;
@@ -260,6 +261,12 @@ protected:
 	/** Adds this function to the suppressed deprecation warnings list for this project */
 	void SuppressDeprecationWarning() const;
 
+	/** Helper function for searching a UFunction for the names of requires pins/params: */
+	static TSet<FName> GetRequiredParamNames(const UFunction* ForFunction);
+
+	/** Routine for validating that all UPARAM(Required) parmas have a connection: */
+	void ValidateRequiredPins(const UFunction* Function, class FCompilerResultsLog& MessageLog) const;
+
 	/** Helper function to find UFunction entries from the skeleton class, use with caution.. */
 	UFunction* GetTargetFunctionFromSkeletonClass() const;
 };
diff --git a/Engine/Source/Editor/BlueprintGraph/Private/K2Node_CallFunction.cpp b/Engine/Source/Editor/BlueprintGraph/Private/K2Node_CallFunction.cpp
--- a/Engine/Source/Editor/BlueprintGraph/Private/K2Node_CallFunction.cpp
+++ b/Engine/Source/Editor/BlueprintGraph/Private/K2Node_CallFunction.cpp
@@ -1118,6 +1118,10 @@ bool UK2Node_CallFunction::CreatePinsForFunctionCall(const UFunction* Function)
 			//Flag pin as read only for const reference property
 			Pin->bDefaultValueIsIgnored = Param->HasAllPropertyFlags(CPF_ConstParm | CPF_ReferenceParm) && (!Function->HasMetaData(FBlueprintMetadata::MD_AutoCreateRefTerm) || Pin->PinType.IsContainer());
 
+			const bool bIsRequiredParam = Param->HasAnyPropertyFlags(CPF_RequiredParm);
+			// Don't let the user edit the default value if the parameter is required to be explicit.
+			Pin->bDefaultValueIsIgnored |= bIsRequiredParam;
+
 			const bool bAdvancedPin = Param->HasAllPropertyFlags(CPF_AdvancedDisplay);
 			Pin->bAdvancedView = bAdvancedPin;
 			if(bAdvancedPin && (ENodeAdvancedPins::NoPins == AdvancedPinDisplay))
@@ -2007,6 +2011,45 @@ void UK2Node_CallFunction::SuppressDeprecationWarning() const
 	}
 }
 
+TSet<FName> UK2Node_CallFunction::GetRequiredParamNames(const UFunction* ForFunction)
+{
+	TSet<FName> Result;
+
+	for (TFieldIterator<FProperty> PropIt(ForFunction); PropIt && (PropIt->PropertyFlags & CPF_Parm); ++PropIt)
+	{
+		FProperty* Param = *PropIt;
+		const bool bIsRequiredParam = Param->HasAnyPropertyFlags(CPF_RequiredParm);
+		if (bIsRequiredParam)
+		{
+			Result.Add(Param->GetFName());
+		}
+	}
+	return Result;
+}
+
+void UK2Node_CallFunction::ValidateRequiredPins(const UFunction* Function, FCompilerResultsLog& MessageLog) const
+{
+	TSet<FName> RequiredPinNames = GetRequiredParamNames(Function);
+
+	if(RequiredPinNames.Num() == 0)
+	{
+		return;
+	}
+
+	for (const UEdGraphPin* Pin : Pins)
+	{
+		if (Pin != nullptr)
+		{
+			const bool bIsRequired = RequiredPinNames.Contains(Pin->GetFName());
+			const bool bIsNotLinked = Pin->LinkedTo.Num() == 0;
+			if (bIsRequired && bIsNotLinked)
+			{
+				MessageLog.Error(*LOCTEXT("MissingRequiredPin", "Pin @@ must be linked to another node (in @@)").ToString(), Pin, this);
+			}
+		}
+	}
+}
+
 void UK2Node_CallFunction::PostPasteNode()
 {
 	Super::PostPasteNode();
@@ -2045,6 +2088,16 @@ void UK2Node_CallFunction::PostPasteNode()
 	}
 }
 
+bool UK2Node_CallFunction::CanSplitPin(const UEdGraphPin* Pin) const
+{
+	TSet<FName> RequiredPins;
+	if (UFunction* Function = GetTargetFunction())
+	{
+		RequiredPins = GetRequiredParamNames(Function);
+	}
+	return Super::CanSplitPin(Pin) && !RequiredPins.Contains(Pin->GetFName());
+}
+
 void UK2Node_CallFunction::PostDuplicate(bool bDuplicateForPIE)
 {
 	Super::PostDuplicate(bDuplicateForPIE);
@@ -2120,6 +2173,8 @@ void UK2Node_CallFunction::ValidateNodeDuringCompilation(class FCompilerResultsL
 
 	if (Function)
 	{
+		ValidateRequiredPins(Function, MessageLog);
+
 		// enforce UnsafeDuringActorConstruction keyword
 		if (Function->HasMetaData(FBlueprintMetadata::MD_UnsafeForConstructionScripts))
 		{
diff --git a/Engine/Source/Programs/Shared/EpicGames.Core/UnrealEngineTypes.cs b/Engine/Source/Programs/Shared/EpicGames.Core/UnrealEngineTypes.cs
--- a/Engine/Source/Programs/Shared/EpicGames.Core/UnrealEngineTypes.cs
+++ b/Engine/Source/Programs/Shared/EpicGames.Core/UnrealEngineTypes.cs
@@ -991,7 +991,10 @@ namespace EpicGames.Core
 		/// </summary>
 		Config = 0x0000000000004000,
 
-		// Unused = 0x0000000000008000,
+		/// <summary>
+		/// Parameter is required in blueprint. Not linking the parameter with a node will result in a compile error.
+		/// </summary>
+		RequiredParm = 0x0000000000008000,
 
 		/// <summary>
 		/// Disable editing on an instance of this class
@@ -1185,7 +1188,7 @@ namespace EpicGames.Core
 		/// <summary>
 		/// All parameter flags
 		/// </summary>
-		ParmFlags = Parm | OutParm | ReturnParm | ReferenceParm | ConstParm,
+		ParmFlags = Parm | OutParm | ReturnParm | ReferenceParm | ConstParm | RequiredParm,
 
 		/// <summary>
 		/// Flags that are propagated to properties inside array container
diff --git a/Engine/Source/Programs/Shared/EpicGames.UHT/Specifiers/UhtPropertyArgumentSpecifiers.cs b/Engine/Source/Programs/Shared/EpicGames.UHT/Specifiers/UhtPropertyArgumentSpecifiers.cs
--- a/Engine/Source/Programs/Shared/EpicGames.UHT/Specifiers/UhtPropertyArgumentSpecifiers.cs
+++ b/Engine/Source/Programs/Shared/EpicGames.UHT/Specifiers/UhtPropertyArgumentSpecifiers.cs
@@ -45,6 +45,13 @@ namespace EpicGames.UHT.Parsers
 				context.MessageSite.LogError("Only parameters in service request functions can be marked NotReplicated");
 			}
 		}
+		
+		[UhtSpecifier(Extends = UhtTableNames.PropertyArgument, ValueType = UhtSpecifierValueType.None)]
+		private static void RequiredSpecifier(UhtSpecifierContext SpecifierContext)
+		{
+			UhtPropertySpecifierContext context = (UhtPropertySpecifierContext)SpecifierContext;
+			context.PropertySettings.PropertyFlags |= EPropertyFlags.RequiredParm;
+		}
 		#endregion
 	}
 }
diff --git a/Engine/Source/Programs/UnrealHeaderTool/Private/HeaderParser.cpp b/Engine/Source/Programs/UnrealHeaderTool/Private/HeaderParser.cpp
--- a/Engine/Source/Programs/UnrealHeaderTool/Private/HeaderParser.cpp
+++ b/Engine/Source/Programs/UnrealHeaderTool/Private/HeaderParser.cpp
@@ -3524,6 +3524,12 @@ void FHeaderParser::GetVarType(
 				}
 				break;
 
+				case EVariableSpecifier::Required:
+				{
+					Flags |= CPF_RequiredParm;
+				}
+				break;
+
 				default:
 				{
 					LogError(TEXT("Unknown variable specifier '%s'"), *Specifier.Key);
diff --git a/Engine/Source/Programs/UnrealHeaderTool/Private/Specifiers/VariableSpecifiers.def b/Engine/Source/Programs/UnrealHeaderTool/Private/Specifiers/VariableSpecifiers.def
--- a/Engine/Source/Programs/UnrealHeaderTool/Private/Specifiers/VariableSpecifiers.def
+++ b/Engine/Source/Programs/UnrealHeaderTool/Private/Specifiers/VariableSpecifiers.def
@@ -37,6 +37,7 @@ VARIABLE_SPECIFIER(Ref)
 VARIABLE_SPECIFIER(Replicated)
 VARIABLE_SPECIFIER(ReplicatedUsing)
 VARIABLE_SPECIFIER(RepRetry)
+VARIABLE_SPECIFIER(Required)
 VARIABLE_SPECIFIER(SaveGame)
 VARIABLE_SPECIFIER(Setter)
 VARIABLE_SPECIFIER(SimpleDisplay)
diff --git a/Engine/Source/Runtime/CoreUObject/Public/UObject/ObjectMacros.h b/Engine/Source/Runtime/CoreUObject/Public/UObject/ObjectMacros.h
--- a/Engine/Source/Runtime/CoreUObject/Public/UObject/ObjectMacros.h
+++ b/Engine/Source/Runtime/CoreUObject/Public/UObject/ObjectMacros.h
@@ -409,7 +409,7 @@ enum EPropertyFlags : uint64
 	CPF_NonNullable						= 0x0000000000001000,	///< Object property can never be null
 	CPF_Transient   					= 0x0000000000002000,	///< Property is transient: shouldn't be saved or loaded, except for Blueprint CDOs.
 	CPF_Config      					= 0x0000000000004000,	///< Property should be loaded/saved as permanent profile.
-	//CPF_								= 0x0000000000008000,	///< 
+	CPF_RequiredParm					= 0x0000000000008000,	///< Parameter must be linked explicitly in blueprint. Leaving the parameter out results in a compile error. 
 	CPF_DisableEditOnInstance			= 0x0000000000010000,	///< Disable editing on an instance of this class
 	CPF_EditConst   					= 0x0000000000020000,	///< Property is uneditable in the editor.
 	CPF_GlobalConfig					= 0x0000000000040000,	///< Load config from base class, not subclass.
@@ -456,7 +456,7 @@ enum EPropertyFlags : uint64
 #define CPF_NativeAccessSpecifiers	(CPF_NativeAccessSpecifierPublic | CPF_NativeAccessSpecifierProtected | CPF_NativeAccessSpecifierPrivate)
 
 /** All parameter flags */
-#define CPF_ParmFlags				(CPF_Parm | CPF_OutParm | CPF_ReturnParm | CPF_ReferenceParm | CPF_ConstParm)
+#define CPF_ParmFlags				(CPF_Parm | CPF_OutParm | CPF_ReturnParm | CPF_RequiredParm | CPF_ReferenceParm | CPF_ConstParm )
 
 /** Flags that are propagated to properties inside containers */
 #define CPF_PropagateToArrayInner	(CPF_ExportObject | CPF_PersistentInstance | CPF_InstancedReference | CPF_ContainsInstancedReference | CPF_Config | CPF_EditConst | CPF_Deprecated | CPF_EditorOnly | CPF_AutoWeak | CPF_UObjectWrapper )
