UE-172222 - Improve redirect support for Blueprints.

A reflected Blueprint type can now be changed without breaking existing content, so long as an autocast function exists for the two types.

Previously, pin reconstruction would give up if two pin types didn't match exactly, which resulted in orphaned pins and compilation errors. Now, it'll first try to find an autocast function for the two types, and insert a cast node into the graph if one is found.

Additionally, notes will be added whenever a new cast node was added, which serves as a notice that users might want to update their API to use the new type change.

diff --git a/Engine/Source/Editor/BlueprintGraph/Classes/K2Node.h b/Engine/Source/Editor/BlueprintGraph/Classes/K2Node.h
--- a/Engine/Source/Editor/BlueprintGraph/Classes/K2Node.h
+++ b/Engine/Source/Editor/BlueprintGraph/Classes/K2Node.h
@@ -447,14 +447,19 @@ public:
 
 protected:
 
-	enum ERedirectType
+	enum ERedirectType : uint8
 	{
 		/* The pins do not match */
-		ERedirectType_None,
+		ERedirectType_None	= 0,
+
 		/* The pins match by name or redirect and have the same type (or we're ok with the mismatched type) */
-		ERedirectType_Name,
+		ERedirectType_Name	= (1 << 0),
+
 		/* The pins match via a redirect and the value needs to also be redirected */
-		ERedirectType_Value
+		ERedirectType_Value	= (1 << 1),
+
+		/* The pin types don't match, but can be connected via a cast node */
+		ERedirectType_Type	= (1 << 2),
 	};
 
 	// Handles the actual reconstruction (copying data, links, name, etc...) from two pins that have already been matched together
diff --git a/Engine/Source/Editor/BlueprintGraph/Private/K2Node.cpp b/Engine/Source/Editor/BlueprintGraph/Private/K2Node.cpp
--- a/Engine/Source/Editor/BlueprintGraph/Private/K2Node.cpp
+++ b/Engine/Source/Editor/BlueprintGraph/Private/K2Node.cpp
@@ -886,23 +886,31 @@ void UK2Node::RestoreSplitPins(TArray<UEdGraphPin*>& OldPins)
 UK2Node::ERedirectType UK2Node::DoPinsMatchForReconstruction(const UEdGraphPin* NewPin, int32 NewPinIndex, const UEdGraphPin* OldPin, int32 OldPinIndex) const
 {
 	ERedirectType RedirectType = ERedirectType_None;
+	const UEdGraphSchema_K2* K2Schema = CastChecked<const UEdGraphSchema_K2>(GetSchema());
 
 	// if the pin names and directions do match
 	if (NewPin->PinName == OldPin->PinName && NewPin->Direction == OldPin->Direction)
 	{
+		const bool bPinsAreTypeCompatible = K2Schema->ArePinTypesCompatible(OldPin->PinType, NewPin->PinType);
+
 		// If the old pin had a default value, only match the new pin if the type is compatible:
-		const UEdGraphSchema_K2* K2Schema = Cast<const UEdGraphSchema_K2>(GetSchema());
-		if (K2Schema &&
-			OldPin->Direction == EGPD_Input &&
+		if (OldPin->Direction == EGPD_Input &&
 			OldPin->LinkedTo.Num() == 0 &&
 			!OldPin->DoesDefaultValueMatchAutogenerated() &&
-			!K2Schema->ArePinTypesCompatible(OldPin->PinType, NewPin->PinType))
+			!bPinsAreTypeCompatible)
 		{
 			RedirectType = ERedirectType_None;
 		}
 		else
 		{
-			RedirectType = ERedirectType_Name;
+			if (bPinsAreTypeCompatible)
+			{
+				RedirectType = ERedirectType_Name;
+			}
+			else if (K2Schema->SearchForAutocastFunction(OldPin->PinType, NewPin->PinType).IsSet())
+			{
+				RedirectType = ERedirectType_Type;
+			}
 		}
 	}
 	else
@@ -979,6 +987,15 @@ UK2Node::ERedirectType UK2Node::DoPinsMatchForReconstruction(const UEdGraphPin*
 				// make sure they match
 				RedirectType = ((NewPin->PinName.ToString() != NewPinNameStr) ? ERedirectType_None : ERedirectType_Name);
 			}
+
+			// A type change might've also occurred with a rename.
+			if (RedirectType == ERedirectType_Name)
+			{
+				if (K2Schema->SearchForAutocastFunction(OldPin->PinType, NewPin->PinType).IsSet())
+				{
+					RedirectType = static_cast<ERedirectType>(static_cast<uint8>(ERedirectType_Name) | static_cast<uint8>(ERedirectType_Type));
+				}
+			}
 		}
 	}
 
@@ -1173,7 +1190,10 @@ void UK2Node::RewireOldPinsToNewPins(TArray<UEdGraphPin*>& InOldPins, TArray<UEd
 				UEdGraphPin* NewPin = InNewPins[NewPinIndex];
 
 				const ERedirectType RedirectType = DoPinsMatchForReconstruction(NewPin, NewPinIndex, OldPin, OldPinIndex);
-				if (RedirectType != ERedirectType_None)
+				const bool bShouldReconstruct = (RedirectType != ERedirectType_None);
+				const bool bNeedsConversion = (RedirectType & ERedirectType_Type);
+
+				if (bShouldReconstruct)
 				{
 					ReconstructSinglePin(NewPin, OldPin, RedirectType);
 					MatchedPins.Add(OldPin, NewPin);
@@ -1183,6 +1203,47 @@ void UK2Node::RewireOldPinsToNewPins(TArray<UEdGraphPin*>& InOldPins, TArray<UEd
 					}
 					bMatched = true;
 					NewPinMatched[NewPinIndex] = true;
+
+					if (bNeedsConversion)
+					{
+						const UEdGraphSchema* Schema = GetSchema();
+						check(Schema);
+
+						TArray<UEdGraphPin*> LinkedToCopy = NewPin->LinkedTo;
+						for (UEdGraphPin* LinkedPin : LinkedToCopy)
+						{
+							// A linked pin could've had its type updated independently of ours.
+							// In some cases, it might've been changed to the same type, which alleviates the need for a cast.
+							if (NewPin->PinType != LinkedPin->PinType)
+							{
+								FFormatNamedArguments Args;
+
+								Args.Add(TEXT("PinName"), FText::FromName(NewPin->GetFName()));
+
+								if (Schema->CreateAutomaticConversionNodeAndConnections(NewPin, LinkedPin))
+								{
+									FText ConversionInfo = FText::Format(
+										NSLOCTEXT("K2Node", "PinConversion_InfoFmt", "Pin '{PinName}' on node @@ changed type. A conversion node had to be inserted into the graph."),
+										Args
+									);
+
+									Message_Note(ConversionInfo.ToString(), this);
+								}
+								else
+								{
+									Args.Add(TEXT("LinkName"), FText::FromName(LinkedPin->GetFName()));
+
+									FText ErrorInfo = FText::Format(
+										NSLOCTEXT("K2Node", "PinConversion_ErrorFmt", "Pin '{PinName}' on node @@ changed type, but no suitable conversion node could be found for link '{LinkName}'!"),
+										Args
+									);
+
+									Message_Error(ErrorInfo.ToString(), this);
+								}
+							}
+						}
+					}
+
 					break;
 				}
 			}
diff --git a/Engine/Source/Editor/BlueprintGraph/Private/K2Node_BreakStruct.cpp b/Engine/Source/Editor/BlueprintGraph/Private/K2Node_BreakStruct.cpp
--- a/Engine/Source/Editor/BlueprintGraph/Private/K2Node_BreakStruct.cpp
+++ b/Engine/Source/Editor/BlueprintGraph/Private/K2Node_BreakStruct.cpp
@@ -349,6 +349,21 @@ UK2Node::ERedirectType UK2Node_BreakStruct::DoPinsMatchForReconstruction(const U
 	{
 		Result = ERedirectType_Name;
 	}
+
+	const UEdGraphSchema_K2* K2Schema = GetDefault<UEdGraphSchema_K2>();
+	check(K2Schema);
+	if (!K2Schema->ArePinTypesCompatible(NewPin->PinType, OldPin->PinType))
+	{
+		if (K2Schema->SearchForAutocastFunction(OldPin->PinType, NewPin->PinType).IsSet())
+		{
+			Result = static_cast<ERedirectType>(static_cast<uint8>(ERedirectType_Name) | static_cast<uint8>(ERedirectType_Type));
+		}
+		else
+		{
+			Result = ERedirectType_None;
+		}
+	}
+
 	return Result;
 }
 
diff --git a/Engine/Source/Editor/BlueprintGraph/Private/K2Node_MakeStruct.cpp b/Engine/Source/Editor/BlueprintGraph/Private/K2Node_MakeStruct.cpp
--- a/Engine/Source/Editor/BlueprintGraph/Private/K2Node_MakeStruct.cpp
+++ b/Engine/Source/Editor/BlueprintGraph/Private/K2Node_MakeStruct.cpp
@@ -316,6 +316,21 @@ UK2Node::ERedirectType UK2Node_MakeStruct::DoPinsMatchForReconstruction(const UE
 	{
 		Result = ERedirectType_Name;
 	}
+
+	const UEdGraphSchema_K2* K2Schema = GetDefault<UEdGraphSchema_K2>();
+	check(K2Schema);
+	if (!K2Schema->ArePinTypesCompatible(NewPin->PinType, OldPin->PinType))
+	{
+		if (K2Schema->SearchForAutocastFunction(OldPin->PinType, NewPin->PinType).IsSet())
+		{
+			Result = static_cast<ERedirectType>(static_cast<uint8>(ERedirectType_Name) | static_cast<uint8>(ERedirectType_Type));
+		}
+		else
+		{
+			Result = ERedirectType_None;
+		}
+	}
+
 	return Result;
 }
 
diff --git a/Engine/Source/Editor/BlueprintGraph/Private/K2Node_StructOperation.cpp b/Engine/Source/Editor/BlueprintGraph/Private/K2Node_StructOperation.cpp
--- a/Engine/Source/Editor/BlueprintGraph/Private/K2Node_StructOperation.cpp
+++ b/Engine/Source/Editor/BlueprintGraph/Private/K2Node_StructOperation.cpp
@@ -75,15 +75,13 @@ bool UK2Node_StructOperation::DoRenamedPinsMatch(const UEdGraphPin* NewPin, cons
 	{
 		const EEdGraphPinDirection StructDirection = bStructInVariablesOut ? EGPD_Input : EGPD_Output;
 		const EEdGraphPinDirection VariablesDirection = bStructInVariablesOut ? EGPD_Output : EGPD_Input;
-		const UEdGraphSchema_K2* K2Schema = GetDefault<UEdGraphSchema_K2>(); 
-		const bool bCompatible = K2Schema && K2Schema->ArePinTypesCompatible(NewPin->PinType, OldPin->PinType);
 
-		if (bCompatible && (StructDirection == OldPin->Direction))
+		if (StructDirection == OldPin->Direction)
 		{
 			// Struct name was changed, which is fine
 			return true;
 		}
-		else if (bCompatible && (VariablesDirection == OldPin->Direction))
+		else if (VariablesDirection == OldPin->Direction)
 		{
 			// Name of a member variable was changed, check guids and redirects
 			if ((NewPin->PersistentGuid == OldPin->PersistentGuid) && OldPin->PersistentGuid.IsValid())
diff --git a/Engine/Source/Editor/BlueprintGraph/Private/K2Node_Variable.cpp b/Engine/Source/Editor/BlueprintGraph/Private/K2Node_Variable.cpp
--- a/Engine/Source/Editor/BlueprintGraph/Private/K2Node_Variable.cpp
+++ b/Engine/Source/Editor/BlueprintGraph/Private/K2Node_Variable.cpp
@@ -343,6 +343,10 @@ UK2Node::ERedirectType UK2Node_Variable::DoPinsMatchForReconstruction( const UEd
 		{
 			return ERedirectType_Name;
 		}
+		else if (K2Schema->SearchForAutocastFunction(OldPin->PinType, NewPin->PinType))
+		{
+			return ERedirectType_Type;
+		}
 		else
 		{
 			const bool bNewPinIsObject = (NewPin->PinType.PinCategory == UEdGraphSchema_K2::PC_Object);
diff --git a/Engine/Source/Runtime/Engine/Classes/EdGraph/EdGraphSchema.h b/Engine/Source/Runtime/Engine/Classes/EdGraph/EdGraphSchema.h
--- a/Engine/Source/Runtime/Engine/Classes/EdGraph/EdGraphSchema.h
+++ b/Engine/Source/Runtime/Engine/Classes/EdGraph/EdGraphSchema.h
@@ -747,22 +747,23 @@ class ENGINE_API UEdGraphSchema : public UObject
 		InputPin = nullptr;
 		OutputPin = nullptr;
 
-		if ((PinA->Direction == EGPD_Input) && (PinB->Direction == EGPD_Output))
+		bool bResult = false;
+		if (ensure(PinA && PinB))
 		{
-			InputPin = PinA;
-			OutputPin = PinB;
-			return true;
-		}
-		else if ((PinB->Direction == EGPD_Input) && (PinA->Direction == EGPD_Output))
-		{
-			InputPin = PinB;
-			OutputPin = PinA;
-			return true;
-		}
-		else
-		{
-			return false;
+			if ((PinA->Direction == EGPD_Input) && (PinB->Direction == EGPD_Output))
+			{
+				InputPin = PinA;
+				OutputPin = PinB;
+				bResult = true;
+			}
+			else if ((PinB->Direction == EGPD_Input) && (PinA->Direction == EGPD_Output))
+			{
+				InputPin = PinB;
+				OutputPin = PinA;
+				bResult = true;
+			}
 		}
+		return bResult;
 	}
 
 	/**
