Making latent BPs more descriptive in Tooltips
Adding Latent BP info to aync task nodes

diff --git a/Engine/Source/Editor/BlueprintGraph/Classes/K2Node_BaseAsyncTask.h b/Engine/Source/Editor/BlueprintGraph/Classes/K2Node_BaseAsyncTask.h
--- a/Engine/Source/Editor/BlueprintGraph/Classes/K2Node_BaseAsyncTask.h
+++ b/Engine/Source/Editor/BlueprintGraph/Classes/K2Node_BaseAsyncTask.h
@@ -65,6 +65,7 @@ class BLUEPRINTGRAPH_API UK2Node_BaseAsyncTask : public UK2Node
 	// UK2Node interface
 	virtual void ExpandNode(FKismetCompilerContext& CompilerContext, UEdGraph* SourceGraph) override;
 	virtual FName GetCornerIcon() const override;
+	virtual FText GetToolTipHeading() const override;
 	virtual void GetMenuActions(FBlueprintActionDatabaseRegistrar& ActionRegistrar) const override;
 	virtual FText GetMenuCategory() const override;
 	/** Whether or not two pins match for purposes of reconnection after reconstruction.  This allows pins that may have had their names changed via reconstruction to be matched to their old values on a node-by-node basis, if needed*/
diff --git a/Engine/Source/Editor/BlueprintGraph/Private/K2Node_BaseAsyncTask.cpp b/Engine/Source/Editor/BlueprintGraph/Private/K2Node_BaseAsyncTask.cpp
--- a/Engine/Source/Editor/BlueprintGraph/Private/K2Node_BaseAsyncTask.cpp
+++ b/Engine/Source/Editor/BlueprintGraph/Private/K2Node_BaseAsyncTask.cpp
@@ -61,8 +61,11 @@ UK2Node_BaseAsyncTask::UK2Node_BaseAsyncTask(const FObjectInitializer& ObjectIni
 
 FText UK2Node_BaseAsyncTask::GetTooltipText() const
 {
-	const FString FunctionToolTipText = UK2Node_CallFunction::GetDefaultTooltipForFunction(GetFactoryFunction());
-	return FText::FromString(FunctionToolTipText);
+	FFormatNamedArguments Args;
+	Args.Add(TEXT("FunctionTooltip"), FText::FromString(UK2Node_CallFunction::GetDefaultTooltipForFunction(GetFactoryFunction())));
+	Args.Add(TEXT("LatentString"), NSLOCTEXT("K2Node", "LatentFunction", "Latent. This node will complete at a later time. Latent nodes can only be placed in event graphs."));
+	
+	return FText::Format(LOCTEXT("AsyncTaskTooltip", "{FunctionTooltip}\n\n{LatentString}"), Args);
 }
 
 FText UK2Node_BaseAsyncTask::GetNodeTitle(ENodeTitleType::Type TitleType) const
@@ -555,6 +558,11 @@ FName UK2Node_BaseAsyncTask::GetCornerIcon() const
 	return TEXT("Graph.Latent.LatentIcon");
 }
 
+FText UK2Node_BaseAsyncTask::GetToolTipHeading() const
+{
+	return LOCTEXT("LatentFunc", "Latent");
+}
+
 FText UK2Node_BaseAsyncTask::GetMenuCategory() const
 {	
 	UFunction* TargetFunction = GetFactoryFunction();
diff --git a/Engine/Source/Editor/BlueprintGraph/Private/K2Node_CallFunction.cpp b/Engine/Source/Editor/BlueprintGraph/Private/K2Node_CallFunction.cpp
--- a/Engine/Source/Editor/BlueprintGraph/Private/K2Node_CallFunction.cpp
+++ b/Engine/Source/Editor/BlueprintGraph/Private/K2Node_CallFunction.cpp
@@ -1532,6 +1532,15 @@ FText UK2Node_CallFunction::GetTooltipText() const
 			// FText::Format() is slow, so we cache this to save on performance
 			CachedTooltip.SetCachedText(FText::Format(LOCTEXT("CallFunction_SubtitledTooltip", "{DefaultTooltip}\n\n{ClientString}"), Args), this);
 		} 
+		else if (Function->HasMetaData(FBlueprintMetadata::MD_Latent))
+		{
+			Args.Add(
+				TEXT("LatentString"),
+				NSLOCTEXT("K2Node", "LatentFunction", "Latent. This node will complete at a later time. Latent nodes can only be placed in event graphs.")
+			);
+			// FText::Format() is slow, so we cache this to save on performance
+			CachedTooltip.SetCachedText(FText::Format(LOCTEXT("CallFunction_SubtitledTooltip", "{DefaultTooltip}\n\n{LatentString}"), Args), this);
+		}
 		else
 		{
 			CachedTooltip.SetCachedText(BaseTooltip, this);
diff --git a/Engine/Source/Editor/BlueprintGraph/Private/K2Node_Event.cpp b/Engine/Source/Editor/BlueprintGraph/Private/K2Node_Event.cpp
--- a/Engine/Source/Editor/BlueprintGraph/Private/K2Node_Event.cpp
+++ b/Engine/Source/Editor/BlueprintGraph/Private/K2Node_Event.cpp
@@ -191,7 +191,16 @@ FText UK2Node_Event::GetTooltipText() const
 				);
 				// FText::Format() is slow, so we cache this to save on performance
 				CachedTooltip.SetCachedText(FText::Format(LOCTEXT("Event_SubtitledTooltip", "{FunctionTooltip}\n\n{ClientString}"), Args), this);
-			}			
+			}
+			else if (Function->HasMetaData(FBlueprintMetadata::MD_Latent))
+			{
+				Args.Add(
+					TEXT("LatentString"),
+					NSLOCTEXT("K2Node", "LatentFunction", "Latent. This node will complete at a later time. Latent nodes can only be placed in event graphs.")
+				);
+				// FText::Format() is slow, so we cache this to save on performance
+				CachedTooltip.SetCachedText(FText::Format(LOCTEXT("CallFunction_SubtitledTooltip", "{DefaultTooltip}\n\n{LatentString}"), Args), this);
+			}
 		}		
 	}
 
@@ -746,6 +755,10 @@ FName UK2Node_Event::GetCornerIcon() const
 			{
 				return TEXT("Graph.Replication.AuthorityOnly");
 			}
+			else if (Function->HasMetaData(FBlueprintMetadata::MD_Latent))
+			{
+				return TEXT("Graph.Latent.LatentIcon");
+			}
 		}
 	}
 
@@ -779,6 +792,10 @@ FText UK2Node_Event::GetToolTipHeading() const
 			{
 				EventHeading = LOCTEXT("ServerOnlyEvent", "Server Only");
 			}
+			else if (Function->HasMetaData(FBlueprintMetadata::MD_Latent))
+			{
+				EventHeading = LOCTEXT("LatentEvent", "Latent");
+			}
 		}
 	}
 
