Fixing horde warning for latent BP tooltips

#rnx

diff --git a/Engine/Source/Editor/BlueprintGraph/Private/K2Node_CallFunction.cpp b/Engine/Source/Editor/BlueprintGraph/Private/K2Node_CallFunction.cpp
--- a/Engine/Source/Editor/BlueprintGraph/Private/K2Node_CallFunction.cpp
+++ b/Engine/Source/Editor/BlueprintGraph/Private/K2Node_CallFunction.cpp
@@ -1517,29 +1517,29 @@ FText UK2Node_CallFunction::GetTooltipText() const
 		if (Function->HasAllFunctionFlags(FUNC_BlueprintAuthorityOnly))
 		{
 			Args.Add(
-				TEXT("ClientString"),
+				TEXT("Subtitle"),
 				NSLOCTEXT("K2Node", "ServerFunction", "Authority Only. This function will only execute on the server.")
 			);
 			// FText::Format() is slow, so we cache this to save on performance
-			CachedTooltip.SetCachedText(FText::Format(LOCTEXT("CallFunction_SubtitledTooltip", "{DefaultTooltip}\n\n{ClientString}"), Args), this);
+			CachedTooltip.SetCachedText(FText::Format(LOCTEXT("CallFunction_SubtitledTooltip", "{DefaultTooltip}\n\n{Subtitle}"), Args), this);
 		}
 		else if (Function->HasAllFunctionFlags(FUNC_BlueprintCosmetic))
 		{
 			Args.Add(
-				TEXT("ClientString"),
+				TEXT("Subtitle"),
 				NSLOCTEXT("K2Node", "ClientFunction", "Cosmetic. This event is only for cosmetic, non-gameplay actions.")
 			);
 			// FText::Format() is slow, so we cache this to save on performance
-			CachedTooltip.SetCachedText(FText::Format(LOCTEXT("CallFunction_SubtitledTooltip", "{DefaultTooltip}\n\n{ClientString}"), Args), this);
+			CachedTooltip.SetCachedText(FText::Format(LOCTEXT("CallFunction_SubtitledTooltip", "{DefaultTooltip}\n\n{Subtitle}"), Args), this);
 		} 
 		else if (Function->HasMetaData(FBlueprintMetadata::MD_Latent))
 		{
 			Args.Add(
-				TEXT("LatentString"),
+				TEXT("Subtitle"),
 				NSLOCTEXT("K2Node", "LatentFunction", "Latent. This node will complete at a later time. Latent nodes can only be placed in event graphs.")
 			);
 			// FText::Format() is slow, so we cache this to save on performance
-			CachedTooltip.SetCachedText(FText::Format(LOCTEXT("CallFunction_SubtitledTooltip", "{DefaultTooltip}\n\n{LatentString}"), Args), this);
+			CachedTooltip.SetCachedText(FText::Format(LOCTEXT("CallFunction_SubtitledTooltip", "{DefaultTooltip}\n\n{Subtitle}"), Args), this);
 		}
 		else
 		{
