UE-174107 - ExpandEnumAsExecs always uses first Exec if an argument name matches Enum value

ExpandEnumAsExecs was using an ambiguous pin lookup. These types of function nodes add additional exec pins that are named after the enumerators. If the target pin specified by ExpandEnumAsExecs happens to have the same name as an enumerator value, then we'll end up with the wrong pin if the basic FindPin function is used.

Instead, we prefer using FindPinByPredicate to specify the exact pin that we're interested in. This ensures that the expansion step creates a SwitchEnum node with a valid UEnum object.

diff --git a/Engine/Source/Editor/BlueprintGraph/Private/K2Node_CallFunction.cpp b/Engine/Source/Editor/BlueprintGraph/Private/K2Node_CallFunction.cpp
--- a/Engine/Source/Editor/BlueprintGraph/Private/K2Node_CallFunction.cpp
+++ b/Engine/Source/Editor/BlueprintGraph/Private/K2Node_CallFunction.cpp
@@ -2512,7 +2512,17 @@ void UK2Node_CallFunction::ExpandNode(class FKismetCompilerContext& CompilerCont
 					Enum = EnumProp->GetEnum();
 				}
 
-				UEdGraphPin* EnumParamPin = FindPin(EnumParamName);
+				auto FindEnumPin = [EnumParamName](const UEdGraphPin* InPin)
+				{
+					check(InPin);
+					const bool bPinMatches =
+						(InPin->PinName == EnumParamName) &&
+						(InPin->PinType.PinCategory == UEdGraphSchema_K2::PC_Byte) &&
+						(Cast<UEnum>(InPin->PinType.PinSubCategoryObject) != nullptr);
+					return bPinMatches;
+				};
+
+				UEdGraphPin* EnumParamPin = FindPinByPredicate(FindEnumPin);
 				if (Enum && EnumParamPin)
 				{
 					// Expanded as input execs pins
