Fixed additional pin reconstruction in promotable operator nodes.

If a promotable operator node had an additional pin (eg: pin C), as well as split pins for any of the default pins (eg: B_X for a vector pin), then ReallocatePinsDuringReconstruction would incorrectly identify a subpin as an additional pin. In this case, the true additional pin would get the wrong pin type.

In practice, this didn't really lead to any issues. However, a future change will make pin type checking a bit stricter, so we need to ensure that recreated pins accurately match their expected types.

The fix here simply modifies GetAdditionalPin to check that the pin isn't a subpin of a default pin.

diff --git a/Engine/Source/Editor/BlueprintGraph/Private/K2Node_PromotableOperator.cpp b/Engine/Source/Editor/BlueprintGraph/Private/K2Node_PromotableOperator.cpp
--- a/Engine/Source/Editor/BlueprintGraph/Private/K2Node_PromotableOperator.cpp
+++ b/Engine/Source/Editor/BlueprintGraph/Private/K2Node_PromotableOperator.cpp
@@ -818,7 +818,15 @@ UEdGraphPin* UK2Node_PromotableOperator::AddInputPinImpl(int32 PinIndex)
 bool UK2Node_PromotableOperator::IsAdditionalPin(const UEdGraphPin* Pin) const
 {
 	// Quickly check if this input pin is one of the two default input pins
-	return Pin && Pin->Direction == EGPD_Input && Pin->PinName != InputPinA_Name && Pin->PinName != InputPinB_Name && !IsTolerancePin(Pin);
+	bool bIsDefaultPinA =
+		(Pin->PinName == InputPinA_Name) ||
+		(Pin->ParentPin && (Pin->ParentPin->PinName == InputPinA_Name));
+
+	bool bIsDefaultPinB =
+		(Pin->PinName == InputPinB_Name) ||
+		(Pin->ParentPin && (Pin->ParentPin->PinName == InputPinB_Name));
+
+	return Pin && (Pin->Direction == EGPD_Input) && !bIsDefaultPinA && !bIsDefaultPinB && !IsTolerancePin(Pin);
 }
 
 bool UK2Node_PromotableOperator::HasAnyConnectionsOrDefaults() const
