Fixed logic in FBindingObject::IsA<>().

diff --git a/Engine/Source/Editor/BlueprintGraph/Classes/BlueprintNodeBinder.h b/Engine/Source/Editor/BlueprintGraph/Classes/BlueprintNodeBinder.h
--- a/Engine/Source/Editor/BlueprintGraph/Classes/BlueprintNodeBinder.h
+++ b/Engine/Source/Editor/BlueprintGraph/Classes/BlueprintNodeBinder.h
@@ -106,7 +106,7 @@ public:
 	template <typename T>
 	typename TEnableIf<!TIsDerivedFrom<T, UObject>::IsDerived, bool>::Type IsA() const
 	{
-		if (bIsUObject && Property.IsValid())
+		if (!bIsUObject && Property.IsValid())
 		{
 			return Property->IsA(T::StaticClass());
 		}
