Fixed static analysis warnings.

diff --git a/Engine/Source/Editor/BlueprintGraph/Classes/K2Node_PromotableOperator.h b/Engine/Source/Editor/BlueprintGraph/Classes/K2Node_PromotableOperator.h
--- a/Engine/Source/Editor/BlueprintGraph/Classes/K2Node_PromotableOperator.h
+++ b/Engine/Source/Editor/BlueprintGraph/Classes/K2Node_PromotableOperator.h
@@ -80,10 +80,10 @@ private:
 	UEdGraphPin* AddInputPinImpl(int32 PinIndex);
 
 	/** Returns true if this pin was added via the IK2Node_AddPinInterface interface */
-	bool IsAdditionalPin(const UEdGraphPin* Pin) const;
+	bool IsAdditionalPin(const UEdGraphPin& Pin) const;
 		
 	/** Returns true if the given pin is a tolerance pin for a comparison operator */
-	bool IsTolerancePin(const UEdGraphPin* Pin) const;
+	bool IsTolerancePin(const UEdGraphPin& Pin) const;
 
 	/** Helper function to recombine all split pins that this node may have */
 	void RecombineAllSplitPins();
diff --git a/Engine/Source/Editor/BlueprintGraph/Private/K2Node_PromotableOperator.cpp b/Engine/Source/Editor/BlueprintGraph/Private/K2Node_PromotableOperator.cpp
--- a/Engine/Source/Editor/BlueprintGraph/Private/K2Node_PromotableOperator.cpp
+++ b/Engine/Source/Editor/BlueprintGraph/Private/K2Node_PromotableOperator.cpp
@@ -265,7 +265,7 @@ bool UK2Node_PromotableOperator::CanConvertPinType(const UEdGraphPin* Pin) const
 		Pin &&
 		!(Pin->Direction == EGPD_Output &&
 		FTypePromotion::IsComparisonFunc(GetTargetFunction())) &&
-		!IsTolerancePin(Pin);
+		!IsTolerancePin(*Pin);
 }
 
 FText UK2Node_PromotableOperator::GetTooltipText() const
@@ -529,7 +529,7 @@ void UK2Node_PromotableOperator::PostReconstructNode()
 
 		for (UEdGraphPin* AddPin : Pins)
 		{
-			if (IsAdditionalPin(AddPin) && AddPin->LinkedTo.Num() > 0)
+			if (ensure(AddPin) && IsAdditionalPin(*AddPin) && AddPin->LinkedTo.Num() > 0)
 			{
 				FEdGraphPinType TypeToSet = FTypePromotion::GetPromotedType(AddPin->LinkedTo);
 				AddPin->PinType = TypeToSet;
@@ -642,23 +642,26 @@ void UK2Node_PromotableOperator::ReallocatePinsDuringReconstruction(TArray<UEdGr
 	// Additional Pin creation here? Check for orphan pins here and see if we can re-create them
 	for (UEdGraphPin* OldPin : OldPins)
 	{
-		if (IsAdditionalPin(OldPin))
+		if (ensure(OldPin))
 		{
-			if (UEdGraphPin* AddPin = GetAdditionalPin(AdditionalPinsFixed + NumFunctionInputs))
+			if (IsAdditionalPin(*OldPin))
 			{
-				AddPin->PinType = OldPin->PinType;
-				AddPin->DefaultValue = OldPin->DefaultValue;
-				++AdditionalPinsFixed;
+				if (UEdGraphPin* AddPin = GetAdditionalPin(AdditionalPinsFixed + NumFunctionInputs))
+				{
+					AddPin->PinType = OldPin->PinType;
+					AddPin->DefaultValue = OldPin->DefaultValue;
+					++AdditionalPinsFixed;
+				}
 			}
-		}
-		
-		// Copy the default value and pin type for pins without any links to preseve it
-		UEdGraphPin* CurrentPin = FindPin(OldPin->GetFName(), OldPin->Direction);
 
-		if(CurrentPin && OldPin->LinkedTo.Num() == 0 && !OldPin->DoesDefaultValueMatchAutogenerated())
-		{
-			CurrentPin->PinType = OldPin->PinType;
-			CurrentPin->DefaultValue = OldPin->DefaultValue;
+			// Copy the default value and pin type for pins without any links to preseve it
+			UEdGraphPin* CurrentPin = FindPin(OldPin->GetFName(), OldPin->Direction);
+
+			if (CurrentPin && OldPin->LinkedTo.Num() == 0 && !OldPin->DoesDefaultValueMatchAutogenerated())
+			{
+				CurrentPin->PinType = OldPin->PinType;
+				CurrentPin->DefaultValue = OldPin->DefaultValue;
+			}
 		}
 	}
 }
@@ -793,9 +796,9 @@ void UK2Node_PromotableOperator::SetFromFunction(const UFunction* Function)
 ///////////////////////////////////////////////////////////
 // UK2Node_PromotableOperator
 
-bool UK2Node_PromotableOperator::IsTolerancePin(const UEdGraphPin* Pin) const
+bool UK2Node_PromotableOperator::IsTolerancePin(const UEdGraphPin& Pin) const
 {
-	return Pin && Pin->PinName == TolerancePin_Name && Pin->Direction == EGPD_Input;
+	return Pin.PinName == TolerancePin_Name && Pin.Direction == EGPD_Input;
 }
 
 UEdGraphPin* UK2Node_PromotableOperator::AddInputPinImpl(int32 PinIndex)
@@ -815,18 +818,18 @@ UEdGraphPin* UK2Node_PromotableOperator::AddInputPinImpl(int32 PinIndex)
 	return NewPin;
 }
 
-bool UK2Node_PromotableOperator::IsAdditionalPin(const UEdGraphPin* Pin) const
+bool UK2Node_PromotableOperator::IsAdditionalPin(const UEdGraphPin& Pin) const
 {
 	// Quickly check if this input pin is one of the two default input pins
 	bool bIsDefaultPinA =
-		(Pin->PinName == InputPinA_Name) ||
-		(Pin->ParentPin && (Pin->ParentPin->PinName == InputPinA_Name));
+		(Pin.PinName == InputPinA_Name) ||
+		(Pin.ParentPin && (Pin.ParentPin->PinName == InputPinA_Name));
 
 	bool bIsDefaultPinB =
-		(Pin->PinName == InputPinB_Name) ||
-		(Pin->ParentPin && (Pin->ParentPin->PinName == InputPinB_Name));
+		(Pin.PinName == InputPinB_Name) ||
+		(Pin.ParentPin && (Pin.ParentPin->PinName == InputPinB_Name));
 
-	return Pin && (Pin->Direction == EGPD_Input) && !bIsDefaultPinA && !bIsDefaultPinB && !IsTolerancePin(Pin);
+	return (Pin.Direction == EGPD_Input) && !bIsDefaultPinA && !bIsDefaultPinB && !IsTolerancePin(Pin);
 }
 
 bool UK2Node_PromotableOperator::HasAnyConnectionsOrDefaults() const
@@ -892,7 +895,7 @@ void UK2Node_PromotableOperator::EvaluatePinsFromChange(UEdGraphPin* ChangedPin,
 		return;
 	}
 	// Changing the tolerance pin doesn't effect the rest of the function signature, so don't attempt to update the func
-	else if (IsTolerancePin(ChangedPin))
+	else if (IsTolerancePin(*ChangedPin))
 	{
 		return;
 	}
@@ -1170,46 +1173,49 @@ void UK2Node_PromotableOperator::GetPinsToConsider(TArray<UEdGraphPin*>& OutArra
 
 	for (UEdGraphPin* Pin : Pins)
 	{
-		// Tolerance pins don't factor into what types we should be matching 
-		// for the function, so we should not consider them
-		if (IsTolerancePin(Pin))
-		{
-			continue;
-		}
-
-		// If this is a comparison operator then we don't need to consider the boolean output
-		// pin because every comparison function will have a boolean output!
-		if (bIsComparisonOp && Pin->Direction == EGPD_Output)
-		{
-			continue;
-		}
-
-		// If this pin has links, then use those instead of the actual pin because we could be process of changing it
-		// which means that it would still have it's old pin type, and could be inaccurate
-		if (Pin->LinkedTo.Num() > 0)
+		if (ensure(Pin))
 		{
-			// If this is from a split pin, then we care about the parent type
-			if (Pin->ParentPin != nullptr)
+			// Tolerance pins don't factor into what types we should be matching 
+			// for the function, so we should not consider them
+			if (IsTolerancePin(*Pin))
 			{
-				OutArray.Add(Pin->ParentPin);
+				continue;
 			}
 
-			// as well as all the links to it
-			for (UEdGraphPin* Link : Pin->LinkedTo)
+			// If this is a comparison operator then we don't need to consider the boolean output
+			// pin because every comparison function will have a boolean output!
+			if (bIsComparisonOp && Pin->Direction == EGPD_Output)
 			{
-				OutArray.Emplace(Link);
+				continue;
 			}
-		}
-		else if (!FWildcardNodeUtils::IsWildcardPin(Pin))
-		{
-			// If this is from a split pin, then we care about the type of the parent, not this pin
-			if (Pin->ParentPin != nullptr)
+
+			// If this pin has links, then use those instead of the actual pin because we could be process of changing it
+			// which means that it would still have it's old pin type, and could be inaccurate
+			if (Pin->LinkedTo.Num() > 0)
 			{
-				OutArray.Add(Pin->ParentPin);
+				// If this is from a split pin, then we care about the parent type
+				if (Pin->ParentPin != nullptr)
+				{
+					OutArray.Add(Pin->ParentPin);
+				}
+
+				// as well as all the links to it
+				for (UEdGraphPin* Link : Pin->LinkedTo)
+				{
+					OutArray.Emplace(Link);
+				}
 			}
-			else
+			else if (!FWildcardNodeUtils::IsWildcardPin(Pin))
 			{
-				OutArray.Add(Pin);
+				// If this is from a split pin, then we care about the type of the parent, not this pin
+				if (Pin->ParentPin != nullptr)
+				{
+					OutArray.Add(Pin->ParentPin);
+				}
+				else
+				{
+					OutArray.Add(Pin);
+				}
 			}
 		}
 	}
@@ -1328,33 +1334,36 @@ void UK2Node_PromotableOperator::UpdatePinsFromFunction(const UFunction* Functio
 	int32 CurPinIndex = 0;
 	for (UEdGraphPin* CurPin : Pins)
 	{
-		// We don't want to try and conform split pin, because we will already have conformed the parent pin
-		if (CurPin->ParentPin != nullptr)
+		if (ensure(CurPin))
 		{
-			continue;
-		}
+			// We don't want to try and conform split pin, because we will already have conformed the parent pin
+			if (CurPin->ParentPin != nullptr)
+			{
+				continue;
+			}
 
-		if (IsAdditionalPin(CurPin))
-		{
-			// Conform to the highest input pin on the function
-			ConformPinLambda(HighestFuncInputType, CurPin);
-		}
-		else if (CurPin->Direction == EGPD_Output)
-		{
-			// Match to the output pin
-			ConformPinLambda(FunctionReturnType, CurPin);
-		}
-		// Creation and conformation of the tolerance pin is handled before all others to 
-		// ensure that connections are not broken accidentally
-		else if (IsTolerancePin(CurPin))
-		{
-			ConformPinLambda(ToleranceType, CurPin);
-		}
-		else
-		{
-			// Match to the appropriate function input type
-			ConformPinLambda(FunctionInputTypes[CurPinIndex], CurPin);
-			++CurPinIndex;
+			if (IsAdditionalPin(*CurPin))
+			{
+				// Conform to the highest input pin on the function
+				ConformPinLambda(HighestFuncInputType, CurPin);
+			}
+			else if (CurPin->Direction == EGPD_Output)
+			{
+				// Match to the output pin
+				ConformPinLambda(FunctionReturnType, CurPin);
+			}
+			// Creation and conformation of the tolerance pin is handled before all others to 
+			// ensure that connections are not broken accidentally
+			else if (IsTolerancePin(*CurPin))
+			{
+				ConformPinLambda(ToleranceType, CurPin);
+			}
+			else
+			{
+				// Match to the appropriate function input type
+				ConformPinLambda(FunctionInputTypes[CurPinIndex], CurPin);
+				++CurPinIndex;
+			}
 		}
 	}
 
