Fixed missing subobject category values for object pins.

In certain cases, non-self object pins have been serialized with null subobject category values. The root cause is due to how UEdGraphSchema_K2::GetPropertyCategoryInfo assigns UClass values. In rare cases, the UClass will be the skeleton class, which is transient, and leads to the null subobject category.

The fix simply modifies GetPropertyCategoryInfo to grab the authoritative class instead. Currently, this issue has only shown up in CallFunction nodes. If older content has this issue, then we'll assign the function's parent class as the subobject.

In a future change, type checking of pins will be a little stricter, so we need to ensure that objects are in a valid state for rewiring to work correctly.

diff --git a/Engine/Source/Editor/BlueprintGraph/Private/EdGraphSchema_K2.cpp b/Engine/Source/Editor/BlueprintGraph/Private/EdGraphSchema_K2.cpp
--- a/Engine/Source/Editor/BlueprintGraph/Private/EdGraphSchema_K2.cpp
+++ b/Engine/Source/Editor/BlueprintGraph/Private/EdGraphSchema_K2.cpp
@@ -88,7 +88,7 @@
 #include "BlueprintNodeBinder.h"
 #include "BlueprintComponentNodeSpawner.h"
 #include "AssetRegistry/AssetRegistryModule.h"
-
+#include "UObject/UE5MainStreamObjectVersion.h"
 #include "K2Node_CastByteToEnum.h"
 #include "K2Node_ClassDynamicCast.h"
 #include "K2Node_GetEnumeratorName.h"
@@ -112,7 +112,7 @@ enum class EBlueprintRealDisplayMode : uint8
 	Number
 };
 
-namespace UE::BlueprintDisplay::Private
+namespace UE::EdGraphSchemaK2::Private
 {
 	int32 LastRealNamingMode = -1;
 	int32 RealNamingMode = static_cast<int32>(EBlueprintRealDisplayMode::Float);
@@ -129,6 +129,51 @@ namespace UE::BlueprintDisplay::Private
 		LastRealNamingMode = RealNamingMode;
 		return bResult;
 	}
+
+	template <class... T>
+	constexpr bool TAlwaysFalse = false;
+
+	template <typename TProperty>
+	UClass* GetAuthoritativeClass(const TProperty& Property)
+	{
+		UClass* PropertyClass = nullptr;
+		if constexpr (std::is_same_v<TProperty, FObjectPropertyBase>)
+		{
+			PropertyClass = Property.PropertyClass;
+		}
+		else if constexpr (std::is_same_v<TProperty, FSoftObjectProperty>)
+		{
+			PropertyClass = Property.PropertyClass;
+		}
+		else if constexpr (std::is_same_v<TProperty, FInterfaceProperty>)
+		{
+			PropertyClass = Property.InterfaceClass;
+		}
+		else if constexpr (std::is_same_v<TProperty, FClassProperty>)
+		{
+			PropertyClass = Property.MetaClass;
+		}
+		else if constexpr (std::is_same_v<TProperty, FSoftClassProperty>)
+		{
+			PropertyClass = Property.MetaClass;
+		}
+		else
+		{
+			static_assert(TAlwaysFalse<TProperty>, "Invalid property used.");
+		}
+
+		if (PropertyClass && PropertyClass->ClassGeneratedBy)
+		{
+			PropertyClass = PropertyClass->GetAuthoritativeClass();
+		}
+
+		if (PropertyClass && FKismetEditorUtilities::IsClassABlueprintSkeleton(PropertyClass))
+		{
+			UE_LOG(LogBlueprint, Warning, TEXT("'%s' is a skeleton class. SubCategoryObject will serialize to a null value."), *PropertyClass->GetFullName());
+		}
+
+		return PropertyClass;
+	}
 }
 
 //////////////////////////////////////////////////////////////////////////
@@ -3389,30 +3434,32 @@ void UEdGraphSchema_K2::CreateFunctionGraphTerminators(UEdGraph& Graph, const UF
 
 bool UEdGraphSchema_K2::GetPropertyCategoryInfo(const FProperty* TestProperty, FName& OutCategory, FName& OutSubCategory, UObject*& OutSubCategoryObject, bool& bOutIsWeakPointer)
 {
+	using namespace UE::EdGraphSchemaK2::Private;
+
 	if (const FInterfaceProperty* InterfaceProperty = CastField<const FInterfaceProperty>(TestProperty))
 	{
 		OutCategory = PC_Interface;
-		OutSubCategoryObject = InterfaceProperty->InterfaceClass;
+		OutSubCategoryObject = GetAuthoritativeClass(*InterfaceProperty);
 	}
 	else if (const FClassProperty* ClassProperty = CastField<const FClassProperty>(TestProperty))
 	{
 		OutCategory = PC_Class;
-		OutSubCategoryObject = ClassProperty->MetaClass;
+		OutSubCategoryObject = GetAuthoritativeClass(*ClassProperty);
 	}
 	else if (const FSoftClassProperty* SoftClassProperty = CastField<const FSoftClassProperty>(TestProperty))
 	{
 		OutCategory = PC_SoftClass;
-		OutSubCategoryObject = SoftClassProperty->MetaClass;
+		OutSubCategoryObject = GetAuthoritativeClass(*SoftClassProperty);
 	}
 	else if (const FSoftObjectProperty* SoftObjectProperty = CastField<const FSoftObjectProperty>(TestProperty))
 	{
 		OutCategory = PC_SoftObject;
-		OutSubCategoryObject = SoftObjectProperty->PropertyClass;
+		OutSubCategoryObject = GetAuthoritativeClass(*SoftObjectProperty);
 	}
 	else if (const FObjectPropertyBase* ObjectProperty = CastField<const FObjectPropertyBase>(TestProperty))
 	{
 		OutCategory = PC_Object;
-		OutSubCategoryObject = ObjectProperty->PropertyClass;
+		OutSubCategoryObject = GetAuthoritativeClass(*ObjectProperty);
 		bOutIsWeakPointer = TestProperty->IsA(FWeakObjectProperty::StaticClass());
 	}
 	else if (const FStructProperty* StructProperty = CastField<const FStructProperty>(TestProperty))
@@ -3725,7 +3772,7 @@ FText UEdGraphSchema_K2::GetCategoryText(const FName Category, const bool bForMe
 
 FText UEdGraphSchema_K2::GetCategoryText(FName Category, FName SubCategory, bool bForMenu)
 {
-	using namespace UE::BlueprintDisplay::Private;
+	using namespace UE::EdGraphSchemaK2::Private;
 
 	if (Category.IsNone())
 	{
@@ -3868,7 +3915,7 @@ FText UEdGraphSchema_K2::TerminalTypeToText(const FName Category, const FName Su
 	{
 		if (Category == UEdGraphSchema_K2::PC_Real)
 		{
-			using namespace UE::BlueprintDisplay::Private;
+			using namespace UE::EdGraphSchemaK2::Private;
 
 			switch (GetRealDisplayMode())
 			{
diff --git a/Engine/Source/Editor/BlueprintGraph/Private/K2Node_CallFunction.cpp b/Engine/Source/Editor/BlueprintGraph/Private/K2Node_CallFunction.cpp
--- a/Engine/Source/Editor/BlueprintGraph/Private/K2Node_CallFunction.cpp
+++ b/Engine/Source/Editor/BlueprintGraph/Private/K2Node_CallFunction.cpp
@@ -8,6 +8,7 @@
 #include "UObject/FrameworkObjectVersion.h"
 #include "UObject/Interface.h"
 #include "UObject/PropertyPortFlags.h"
+#include "UObject/UE5MainStreamObjectVersion.h"
 #include "Kismet/BlueprintFunctionLibrary.h"
 #include "Engine/BlueprintGeneratedClass.h"
 #include "Framework/Notifications/NotificationManager.h"
@@ -2286,6 +2287,7 @@ void UK2Node_CallFunction::Serialize(FArchive& Ar)
 	Super::Serialize(Ar);
 
 	Ar.UsingCustomVersion(FReleaseObjectVersion::GUID);
+	Ar.UsingCustomVersion(FUE5MainStreamObjectVersion::GUID);
 
 	if (Ar.IsLoading())
 	{
@@ -2355,6 +2357,25 @@ void UK2Node_CallFunction::Serialize(FArchive& Ar)
 			}
 		}
 
+		for (int32 PinIndex = 0; PinIndex < Pins.Num(); ++PinIndex)
+		{
+			UEdGraphPin* Pin = Pins[PinIndex];
+			check(Pin);
+
+			bool bNeedsSubCategoryObjectRepair =
+				(Pin->PinType.PinCategory == UEdGraphSchema_K2::PC_Object) &&
+				(Pin->PinType.PinSubCategory != UEdGraphSchema_K2::PSC_Self) &&
+				(Pin->PinType.PinSubCategoryObject == nullptr) &&
+				(Ar.CustomVer(FReleaseObjectVersion::GUID) < FUE5MainStreamObjectVersion::NullPinSubCategoryObjectFix);
+
+			// Prior to NullPinSubCategoryObjectFix, some object pins were serialized with a null PinSubCategoryObject.
+			// Going forward, this will be an error, so we'll attempt to repair the pin by assigning a class.
+			if (bNeedsSubCategoryObjectRepair)
+			{
+				Pin->PinType.PinSubCategoryObject = FunctionReference.GetMemberParentClass();
+			}
+		}
+
 		if (!Ar.IsObjectReferenceCollector())
 		{
 			// Don't validate the enabled state if the user has explicitly set it. Also skip validation if we're just duplicating this node.
diff --git a/Engine/Source/Runtime/Core/Public/UObject/UE5MainStreamObjectVersion.h b/Engine/Source/Runtime/Core/Public/UObject/UE5MainStreamObjectVersion.h
--- a/Engine/Source/Runtime/Core/Public/UObject/UE5MainStreamObjectVersion.h
+++ b/Engine/Source/Runtime/Core/Public/UObject/UE5MainStreamObjectVersion.h
@@ -328,6 +328,9 @@ struct CORE_API FUE5MainStreamObjectVersion
 		// RigVM generated class refactor
 		RigVMGeneratedClass,
 
+		// In certain cases, Blueprint pins with a PC_Object category would serialize a null PinSubCategoryObject
+		NullPinSubCategoryObjectFix,
+
 		// -----<new versions can be added above this line>-------------------------------------------------
 		VersionPlusOne,
 		LatestVersion = VersionPlusOne - 1
