[BugFix] UK2Node_SpawnActorFromClass::FixupScaleMethodPin has non-deterministic behavior during cook time
 - fixed by breaking out early during cook. FixupScaleMethodPin shouldn't have any impact on runtime behavior.

diff --git a/Engine/Source/Editor/BlueprintGraph/Private/K2Node_SpawnActorFromClass.cpp b/Engine/Source/Editor/BlueprintGraph/Private/K2Node_SpawnActorFromClass.cpp
--- a/Engine/Source/Editor/BlueprintGraph/Private/K2Node_SpawnActorFromClass.cpp
+++ b/Engine/Source/Editor/BlueprintGraph/Private/K2Node_SpawnActorFromClass.cpp
@@ -379,6 +379,13 @@ void UK2Node_SpawnActorFromClass::FixupScaleMethodPin()
 		return;
 	}
 
+	// For the sake of cook determinism, leave the scale method as ESpawnActorScaleMethod::SelectDefaultAtRuntime when cooking.
+	// This shouldn't impact runtime behavior.
+	if (IsRunningCookCommandlet())
+	{
+		return;
+	}
+
 	if (GetLinkerCustomVersion(FUE5MainStreamObjectVersion::GUID) < FUE5MainStreamObjectVersion::SpawnActorFromClassTransformScaleMethod)
 	{
         if (UEdGraphPin* const ScaleMethodPin = TryGetScaleMethodPin())
