UE-174107 - ExpandEnumAsExecs always uses first Exec if an argument name matches Enum value (resubmit)

ExpandEnumAsExecs was using an ambiguous pin lookup. These types of function nodes add additional exec pins that are named after the enumerators. If the target pin specified by ExpandEnumAsExecs happens to have the same name as an enumerator value, then we'll end up with the wrong pin if the basic FindPin function is used.

Instead, we prefer using FindPinByPredicate to specify the exact pin that we're interested in. This ensures that the expansion step creates a SwitchEnum node with a valid UEnum object.

diff --git a/Engine/Source/Editor/BlueprintGraph/Private/K2Node_CallFunction.cpp b/Engine/Source/Editor/BlueprintGraph/Private/K2Node_CallFunction.cpp
--- a/Engine/Source/Editor/BlueprintGraph/Private/K2Node_CallFunction.cpp
+++ b/Engine/Source/Editor/BlueprintGraph/Private/K2Node_CallFunction.cpp
@@ -52,6 +52,38 @@
 
 #define LOCTEXT_NAMESPACE "K2Node"
 
+namespace UE::K2NodeCallFunction::Private
+{
+	UEdGraphPin* FindBoolParamPin(const UK2Node_CallFunction& Node, FName ParameterName)
+	{
+		auto FindPin = [ParameterName](const UEdGraphPin* InPin)
+		{
+			check(InPin);
+			const bool bPinMatches =
+				(InPin->PinName == ParameterName) &&
+				(InPin->PinType.PinCategory == UEdGraphSchema_K2::PC_Boolean);
+			return bPinMatches;
+		};
+
+		return Node.FindPinByPredicate(FindPin);
+	}
+
+	UEdGraphPin* FindEnumParamPin(const UK2Node_CallFunction& Node, FName ParameterName)
+	{
+		auto FindPin = [ParameterName](const UEdGraphPin* InPin)
+		{
+			check(InPin);
+			const bool bPinMatches =
+				(InPin->PinName == ParameterName) &&
+				(InPin->PinType.PinCategory == UEdGraphSchema_K2::PC_Byte) &&
+				(Cast<UEnum>(InPin->PinType.PinSubCategoryObject) != nullptr);
+			return bPinMatches;
+		};
+
+		return Node.FindPinByPredicate(FindPin);
+	}
+}
+
 /*******************************************************************************
  *  FCustomStructureParamHelper
  ******************************************************************************/
@@ -777,7 +809,7 @@ void UK2Node_CallFunction::CreateExecPinsForFunctionCall(const UFunction* Functi
 	if (!bIsPureFunc)
 	{
 		// If we want enum->exec expansion, and it is not disabled, do it now
-		if(bWantsEnumToExecExpansion)
+		if (bWantsEnumToExecExpansion)
 		{
 			TArray<FName> EnumNames;
 			GetExpandEnumPinNames(Function, EnumNames);
@@ -1161,15 +1193,17 @@ bool UK2Node_CallFunction::CreatePinsForFunctionCall(const UFunction* Function)
 	}
 
 	// If we have 'enum to exec' parameters, set their default value to something valid so we don't get warnings
-	if(bWantsEnumToExecExpansion)
+	if (bWantsEnumToExecExpansion)
 	{
+		using namespace UE::K2NodeCallFunction::Private;
+
 		TArray<FName> EnumNamesToCheck;
 		GetExpandEnumPinNames(Function, EnumNamesToCheck);
 
 		for (const FName& Name : EnumNamesToCheck)
 		{
-			UEdGraphPin* EnumParamPin = FindPin(Name);
-			if (UEnum* PinEnum = (EnumParamPin ? Cast<UEnum>(EnumParamPin->PinType.PinSubCategoryObject.Get()) : NULL))
+			UEdGraphPin* EnumParamPin = FindEnumParamPin(*this, Name);
+			if (UEnum* PinEnum = (EnumParamPin ? Cast<UEnum>(EnumParamPin->PinType.PinSubCategoryObject.Get()) : nullptr))
 			{
 				EnumParamPin->DefaultValue = PinEnum->GetNameStringByIndex(0);
 			}
@@ -2468,10 +2502,12 @@ void UK2Node_CallFunction::ExpandNode(class FKismetCompilerContext& CompilerCont
 	}
 
 	// If we have an enum param that is expanded, we handle that first
-	if(bWantsEnumToExecExpansion)
+	if (bWantsEnumToExecExpansion)
 	{
-		if(Function)
+		if (Function)
 		{
+			using namespace UE::K2NodeCallFunction::Private;
+
 			TArray<FName> EnumNamesToCheck;
 			GetExpandEnumPinNames(Function, EnumNamesToCheck);
 
@@ -2520,20 +2556,26 @@ void UK2Node_CallFunction::ExpandNode(class FKismetCompilerContext& CompilerCont
 				}
 			};
 
-			for (const FName& EnumParamName : EnumNamesToCheck)
+			for (const FName& ParamName : EnumNamesToCheck)
 			{
 				UEnum* Enum = nullptr;
-
-				if (FByteProperty* ByteProp = FindFProperty<FByteProperty>(Function, EnumParamName))
+				UEdGraphPin* EnumParamPin = nullptr;
+	
+				if (FBoolProperty* BoolProp = FindFProperty<FBoolProperty>(Function, ParamName))
+				{
+					EnumParamPin = FindBoolParamPin(*this, ParamName);
+				}
+				else if (FByteProperty* ByteProp = FindFProperty<FByteProperty>(Function, ParamName))
 				{
 					Enum = ByteProp->Enum;
+					EnumParamPin = FindEnumParamPin(*this, ParamName);
 				}
-				else if (FEnumProperty* EnumProp = FindFProperty<FEnumProperty>(Function, EnumParamName))
+				else if (FEnumProperty* EnumProp = FindFProperty<FEnumProperty>(Function, ParamName))
 				{
 					Enum = EnumProp->GetEnum();
+					EnumParamPin = FindEnumParamPin(*this, ParamName);
 				}
 
-				UEdGraphPin* EnumParamPin = FindPin(EnumParamName);
 				if (Enum && EnumParamPin)
 				{
 					// Expanded as input execs pins
