#UE Performance optimization. When updating FAutocastFunctionMap as a result of a module loading, only iterate through the function libraries in the module, instead of all function libraries in memory

diff --git a/Engine/Source/Editor/BlueprintGraph/Private/EdGraphSchema_K2.cpp b/Engine/Source/Editor/BlueprintGraph/Private/EdGraphSchema_K2.cpp
--- a/Engine/Source/Editor/BlueprintGraph/Private/EdGraphSchema_K2.cpp
+++ b/Engine/Source/Editor/BlueprintGraph/Private/EdGraphSchema_K2.cpp
@@ -2441,10 +2441,47 @@ public:
 #endif //SCHEMA_K2_AUTOCASTFUNCTIONMAP_LOG_TIME
 
 		InnerMap.Empty();
-		const UEdGraphSchema_K2* Schema = GetDefault<UEdGraphSchema_K2>();
 
 		TArray<UClass*> Libraries;
 		GetDerivedClasses(UBlueprintFunctionLibrary::StaticClass(), Libraries);
+		AddLibraries(Libraries);
+
+#if SCHEMA_K2_AUTOCASTFUNCTIONMAP_LOG_TIME
+		const double EndTime = FPlatformTime::Seconds();
+		UE_LOG(LogBlueprint, Warning, TEXT("FAutocastFunctionMap::Refresh took %fs"), EndTime - StartTime);
+#endif //SCHEMA_K2_AUTOCASTFUNCTIONMAP_LOG_TIME
+#undef SCHEMA_K2_AUTOCASTFUNCTIONMAP_LOG_TIME
+	}
+
+	void AddLibrariesFromModule(FName ModuleThatChanged)
+	{
+		TRACE_CPUPROFILER_EVENT_SCOPE(WILD_FAutocastFunctionMap::AddLibrariesFromModule);
+
+		const UEdGraphSchema_K2* Schema = GetDefault<UEdGraphSchema_K2>();
+
+		if (UPackage* ModuleScriptPacakge = FindPackage(nullptr, *FString::Printf(TEXT("/Script/%s"), *ModuleThatChanged.ToString())))
+		{
+			TArray<UClass*> Libraries;
+			ForEachObjectWithPackage(ModuleScriptPacakge, [&Libraries](UObject* Obj) -> bool
+				{
+					if (UClass* Class = Cast<UClass>(Obj))
+					{
+						if (Class->IsChildOf(UBlueprintFunctionLibrary::StaticClass()))
+						{
+							Libraries.Add(Class);
+						}
+					}
+
+					return true;
+				}, false);
+
+			AddLibraries(Libraries);
+		}
+	}
+
+	void AddLibraries(const TArray<UClass*>& Libraries)
+	{
+		const UEdGraphSchema_K2* Schema = GetDefault<UEdGraphSchema_K2>();
 		for (UClass* Library : Libraries)
 		{
 			if (Library && (CLASS_Native == (Library->ClassFlags & (CLASS_Native | CLASS_Deprecated | CLASS_NewerVersionExists))))
@@ -2458,12 +2495,6 @@ public:
 				}
 			}
 		}
-
-#if SCHEMA_K2_AUTOCASTFUNCTIONMAP_LOG_TIME
-		const double EndTime = FPlatformTime::Seconds();
-		UE_LOG(LogBlueprint, Warning, TEXT("FAutocastFunctionMap::Refresh took %fs"), EndTime - StartTime);
-#endif //SCHEMA_K2_AUTOCASTFUNCTIONMAP_LOG_TIME
-#undef SCHEMA_K2_AUTOCASTFUNCTIONMAP_LOG_TIME
 	}
 
 	UFunction* Find(const FEdGraphPinType& InputPinType, const FEdGraphPinType& OutputPinType) const
@@ -2499,7 +2530,14 @@ public:
 	{
 		if (AutocastFunctionMap)
 		{
-			AutocastFunctionMap->Refresh();
+			if (ReasonForChange == EModuleChangeReason::ModuleLoaded)
+			{
+				AutocastFunctionMap->AddLibrariesFromModule(ModuleThatChanged);
+			}
+			else if (ReasonForChange == EModuleChangeReason::ModuleUnloaded)
+			{
+				AutocastFunctionMap->Refresh();
+			}
 		}
 	}
 
