Generated Blueprint skeleton class functions will no longer parent themselves to interface class functions in the fast path, while also ensuring backwards-compatibility for existing call sites to implemented interface functions.

Epidemiology:
- There was a regression introduced with the change to BPCM in 4.17 which started causing all newly-placed function call nodes to use the interface class as the Target pin type, rather than the 'self' type.
- However, the compiled function context was still inferred as the 'self' type, which resulted in emitting to bytecode only EX_Context w/o the required EX_InterfaceContext.
- This meant at runtime the VM would read an FScriptInterface value (16 bytes) into a UObject* ptr value (8 bytes) allocated on the stack (execContext), which caused a stack overrun (recently surfaced by ASAN).
- That regression is now fixed in the BPCM, however..
- Since we also allow object->interface connections to pass w/o casting when the object type implements the interface, this appeared to work the same as before the regression from a user perspective. However, it created an inconsistency for Target pin types that led to some odd UX; for example, if you choose to no longer implement the interface, now you might have a local function call site w/ an interface-typed Target pin that you can no longer connect self to/broken connection.
- Fixing the BPCM issue also meant that Target pins reverted back to 'self' object type on existing nodes, which caused a backcompat issue from being linked to interface pins (as that requires an interface->object conversion and will fail validation).
- Fixing that *seemed* relatively straightforward; I added a bit of code to modify the node expansion to spawn an intermediate cast node. However..
- Discovered that dynamic cast nodes had an issue where ones we dropped as an autocast via TryCreateConnection() would not be pure (i.e. no exec pins) even though it was explicitly being set in the autocast logic.
- This meant that any intermediate autocast nodes for interface->self end up being dropped as impure, with no connection to the exec pins, so they would end up being pruned at compile time. (this problem was due to another regression that traced back a long long ways..it was introduced in UE 4.6).
- Fixing that required a change to the way in which Dynamic Cast nodes manage the internal pure/impure state (since it can be toggled by the user, and it also has a default setting for new nodes).
- After fixing that, existing connections now continue to work (including previously-connected interface array outputs), while dragging new connections from interface->self will now add an explicit autocast node.
- An additional fix was made to prevent users from connecting an array of interfaces output to a function call node's Target pin for new connections (existing connections will continue to function per above).
--> Since call nodes support a ForEach-style expansion, the array connection was being allowed to pass, but then it would try to add an autocast node as if it were a scalar type, which doesn't support the expansion, so it could not be connected (broken UX).

diff --git a/Engine/Source/Editor/BlueprintGraph/Classes/K2Node_DynamicCast.h b/Engine/Source/Editor/BlueprintGraph/Classes/K2Node_DynamicCast.h
--- a/Engine/Source/Editor/BlueprintGraph/Classes/K2Node_DynamicCast.h
+++ b/Engine/Source/Editor/BlueprintGraph/Classes/K2Node_DynamicCast.h
@@ -34,6 +34,10 @@ class UK2Node_DynamicCast : public UK2Node
 	UPROPERTY()
 	TSubclassOf<class UObject>  TargetType;
 
+	//~ Begin UObject Interface
+	virtual void Serialize(FArchive& Ar);
+	//~ End UObject Interface
+
 	//~ Begin UEdGraphNode Interface
 	virtual void AllocateDefaultPins() override;
 	virtual FLinearColor GetNodeTitleColor() const override;
@@ -43,7 +47,6 @@ class UK2Node_DynamicCast : public UK2Node
 	virtual void GetNodeContextMenuActions(class UToolMenu* Menu, class UGraphNodeContextMenuContext* Context) const override;
 	virtual bool IncludeParentNodeContextMenu() const override { return true; }
 	virtual void PostReconstructNode() override;
-	virtual void PostPlacedNewNode() override;
 	virtual bool HasExternalDependencies(TArray<class UStruct*>* OptionalOutput) const override;
 	//~ End UEdGraphNode Interface
 
@@ -52,7 +55,7 @@ class UK2Node_DynamicCast : public UK2Node
 	virtual class FNodeHandlingFunctor* CreateNodeHandler(class FKismetCompilerContext& CompilerContext) const override;
 	virtual FText GetMenuCategory() const override;
 	virtual FBlueprintNodeSignature GetSignature() const override;
-	virtual bool IsNodePure() const override { return bIsPureCast; }
+	virtual bool IsNodePure() const override { return PureState == EPureState::Pure; }
 	virtual bool IsConnectionDisallowed(const UEdGraphPin* MyPin, const UEdGraphPin* OtherPin, FString& OutReason) const override;
 	virtual void NotifyPinConnectionListChanged(UEdGraphPin* Pin) override;
 	virtual void ReallocatePinsDuringReconstruction(TArray<UEdGraphPin*>& OldPins) override;
@@ -79,11 +82,17 @@ class UK2Node_DynamicCast : public UK2Node
 	 * Will change the node's purity, and reallocate pins accordingly (adding/
 	 * removing exec pins).
 	 * 
-	 * @param  bNewPurity  The new value for bIsPureCast.
+	 * @param  bNewPurity  Whether to configure node as pure.
 	 */
 	BLUEPRINTGRAPH_API void SetPurity(bool bNewPurity);
 
 protected:
+	/** Determine purity and create exec pins if needed */
+	void CreateExecPins();
+
+	/** Create success output signal pin (visible only if pure) */
+	void CreateSuccessPin();
+
 	/** Flips the node's purity (adding/removing exec pins as needed). */
 	void TogglePurity();
 
@@ -93,7 +102,23 @@ protected:
 	/** Constructing FText strings can be costly, so we cache the node's title */
 	FNodeTextCache CachedNodeTitle;
 
-	UPROPERTY()
+	/** [DEPRECATED] This member is no longer in use. */
+	UE_DEPRECATED(5.3, "Use IsNodePure() instead.")
 	bool bIsPureCast;
+
+private:
+	enum class EPureState : uint8
+	{
+		Pure,
+		Impure,
+		UseDefault
+	};
+
+	/** Internal state used to determine purity */
+	EPureState PureState;
+
+	/** [DEPRECATED] Receives the old value on load if previously saved */
+	UPROPERTY()
+	bool bIsPureCast_DEPRECATED;
 };
 
diff --git a/Engine/Source/Editor/BlueprintGraph/Private/EdGraphSchema_K2.cpp b/Engine/Source/Editor/BlueprintGraph/Private/EdGraphSchema_K2.cpp
--- a/Engine/Source/Editor/BlueprintGraph/Private/EdGraphSchema_K2.cpp
+++ b/Engine/Source/Editor/BlueprintGraph/Private/EdGraphSchema_K2.cpp
@@ -88,7 +88,7 @@
 #include "BlueprintNodeBinder.h"
 #include "BlueprintComponentNodeSpawner.h"
 #include "AssetRegistry/AssetRegistryModule.h"
-#include "UObject/UE5MainStreamObjectVersion.h"
+
 #include "K2Node_CastByteToEnum.h"
 #include "K2Node_ClassDynamicCast.h"
 #include "K2Node_GetEnumeratorName.h"
@@ -112,7 +112,7 @@ enum class EBlueprintRealDisplayMode : uint8
 	Number
 };
 
-namespace UE::EdGraphSchemaK2::Private
+namespace UE::BlueprintDisplay::Private
 {
 	int32 LastRealNamingMode = -1;
 	int32 RealNamingMode = static_cast<int32>(EBlueprintRealDisplayMode::Float);
@@ -129,51 +129,6 @@ namespace UE::EdGraphSchemaK2::Private
 		LastRealNamingMode = RealNamingMode;
 		return bResult;
 	}
-
-	template <class... T>
-	constexpr bool TAlwaysFalse = false;
-
-	template <typename TProperty>
-	UClass* GetAuthoritativeClass(const TProperty& Property)
-	{
-		UClass* PropertyClass = nullptr;
-		if constexpr (std::is_same_v<TProperty, FObjectPropertyBase>)
-		{
-			PropertyClass = Property.PropertyClass;
-		}
-		else if constexpr (std::is_same_v<TProperty, FSoftObjectProperty>)
-		{
-			PropertyClass = Property.PropertyClass;
-		}
-		else if constexpr (std::is_same_v<TProperty, FInterfaceProperty>)
-		{
-			PropertyClass = Property.InterfaceClass;
-		}
-		else if constexpr (std::is_same_v<TProperty, FClassProperty>)
-		{
-			PropertyClass = Property.MetaClass;
-		}
-		else if constexpr (std::is_same_v<TProperty, FSoftClassProperty>)
-		{
-			PropertyClass = Property.MetaClass;
-		}
-		else
-		{
-			static_assert(TAlwaysFalse<TProperty>, "Invalid property used.");
-		}
-
-		if (PropertyClass && PropertyClass->ClassGeneratedBy)
-		{
-			PropertyClass = PropertyClass->GetAuthoritativeClass();
-		}
-
-		if (PropertyClass && FKismetEditorUtilities::IsClassABlueprintSkeleton(PropertyClass))
-		{
-			UE_LOG(LogBlueprint, Warning, TEXT("'%s' is a skeleton class. SubCategoryObject will serialize to a null value."), *PropertyClass->GetFullName());
-		}
-
-		return PropertyClass;
-	}
 }
 
 //////////////////////////////////////////////////////////////////////////
@@ -2731,7 +2686,7 @@ TOptional<UEdGraphSchema_K2::FFindSpecializedConversionNodeResults> UEdGraphSche
 	const bool bTryAlternateObjectProperty =
 		InputPin.GetOwningNode()->IsA(UK2Node_CallFunction::StaticClass()) &&
 		IsSelfPin(InputPin) &&
-		((OutputPinType.PinCategory == PC_Object) || (OutputPinType.PinCategory == PC_Interface));
+		((OutputPinType.PinCategory == PC_Object) || ((OutputPinType.PinCategory == PC_Interface) && !OutputPinType.IsContainer()));
 
 	if (bConvertScalarToArray)
 	{
@@ -2809,8 +2764,9 @@ TOptional<UEdGraphSchema_K2::FFindSpecializedConversionNodeResults> UEdGraphSche
 				Result->TargetNode = CastByteToEnum;
 			}
 		}
-		else
+		else if (!OutputPinType.IsContainer())
 		{
+			// Note: The nodes here do not support a ForEach-style expansion, so we exclude container types.
 			UClass* InputClass  = FBlueprintEditorUtils::GetTypeForPin(InputPin);
 			const UClass* OutputClass = Cast<UClass>(OutputPinType.PinSubCategoryObject.Get());
 
@@ -3472,32 +3428,30 @@ void UEdGraphSchema_K2::CreateFunctionGraphTerminators(UEdGraph& Graph, const UF
 
 bool UEdGraphSchema_K2::GetPropertyCategoryInfo(const FProperty* TestProperty, FName& OutCategory, FName& OutSubCategory, UObject*& OutSubCategoryObject, bool& bOutIsWeakPointer)
 {
-	using namespace UE::EdGraphSchemaK2::Private;
-
 	if (const FInterfaceProperty* InterfaceProperty = CastField<const FInterfaceProperty>(TestProperty))
 	{
 		OutCategory = PC_Interface;
-		OutSubCategoryObject = GetAuthoritativeClass(*InterfaceProperty);
+		OutSubCategoryObject = InterfaceProperty->InterfaceClass;
 	}
 	else if (const FClassProperty* ClassProperty = CastField<const FClassProperty>(TestProperty))
 	{
 		OutCategory = PC_Class;
-		OutSubCategoryObject = GetAuthoritativeClass(*ClassProperty);
+		OutSubCategoryObject = ClassProperty->MetaClass;
 	}
 	else if (const FSoftClassProperty* SoftClassProperty = CastField<const FSoftClassProperty>(TestProperty))
 	{
 		OutCategory = PC_SoftClass;
-		OutSubCategoryObject = GetAuthoritativeClass(*SoftClassProperty);
+		OutSubCategoryObject = SoftClassProperty->MetaClass;
 	}
 	else if (const FSoftObjectProperty* SoftObjectProperty = CastField<const FSoftObjectProperty>(TestProperty))
 	{
 		OutCategory = PC_SoftObject;
-		OutSubCategoryObject = GetAuthoritativeClass(*SoftObjectProperty);
+		OutSubCategoryObject = SoftObjectProperty->PropertyClass;
 	}
 	else if (const FObjectPropertyBase* ObjectProperty = CastField<const FObjectPropertyBase>(TestProperty))
 	{
 		OutCategory = PC_Object;
-		OutSubCategoryObject = GetAuthoritativeClass(*ObjectProperty);
+		OutSubCategoryObject = ObjectProperty->PropertyClass;
 		bOutIsWeakPointer = TestProperty->IsA(FWeakObjectProperty::StaticClass());
 	}
 	else if (const FStructProperty* StructProperty = CastField<const FStructProperty>(TestProperty))
@@ -3810,7 +3764,7 @@ FText UEdGraphSchema_K2::GetCategoryText(const FName Category, const bool bForMe
 
 FText UEdGraphSchema_K2::GetCategoryText(FName Category, FName SubCategory, bool bForMenu)
 {
-	using namespace UE::EdGraphSchemaK2::Private;
+	using namespace UE::BlueprintDisplay::Private;
 
 	if (Category.IsNone())
 	{
@@ -3953,7 +3907,7 @@ FText UEdGraphSchema_K2::TerminalTypeToText(const FName Category, const FName Su
 	{
 		if (Category == UEdGraphSchema_K2::PC_Real)
 		{
-			using namespace UE::EdGraphSchemaK2::Private;
+			using namespace UE::BlueprintDisplay::Private;
 
 			switch (GetRealDisplayMode())
 			{
diff --git a/Engine/Source/Editor/BlueprintGraph/Private/K2Node_CallFunction.cpp b/Engine/Source/Editor/BlueprintGraph/Private/K2Node_CallFunction.cpp
--- a/Engine/Source/Editor/BlueprintGraph/Private/K2Node_CallFunction.cpp
+++ b/Engine/Source/Editor/BlueprintGraph/Private/K2Node_CallFunction.cpp
@@ -8,7 +8,6 @@
 #include "UObject/FrameworkObjectVersion.h"
 #include "UObject/Interface.h"
 #include "UObject/PropertyPortFlags.h"
-#include "UObject/UE5MainStreamObjectVersion.h"
 #include "Kismet/BlueprintFunctionLibrary.h"
 #include "Engine/BlueprintGeneratedClass.h"
 #include "Framework/Notifications/NotificationManager.h"
@@ -52,38 +51,6 @@
 
 #define LOCTEXT_NAMESPACE "K2Node"
 
-namespace UE::K2NodeCallFunction::Private
-{
-	UEdGraphPin* FindBoolParamPin(const UK2Node_CallFunction& Node, FName ParameterName)
-	{
-		auto FindPin = [ParameterName](const UEdGraphPin* InPin)
-		{
-			check(InPin);
-			const bool bPinMatches =
-				(InPin->PinName == ParameterName) &&
-				(InPin->PinType.PinCategory == UEdGraphSchema_K2::PC_Boolean);
-			return bPinMatches;
-		};
-
-		return Node.FindPinByPredicate(FindPin);
-	}
-
-	UEdGraphPin* FindEnumParamPin(const UK2Node_CallFunction& Node, FName ParameterName)
-	{
-		auto FindPin = [ParameterName](const UEdGraphPin* InPin)
-		{
-			check(InPin);
-			const bool bPinMatches =
-				(InPin->PinName == ParameterName) &&
-				(InPin->PinType.PinCategory == UEdGraphSchema_K2::PC_Byte) &&
-				(Cast<UEnum>(InPin->PinType.PinSubCategoryObject) != nullptr);
-			return bPinMatches;
-		};
-
-		return Node.FindPinByPredicate(FindPin);
-	}
-}
-
 /*******************************************************************************
  *  FCustomStructureParamHelper
  ******************************************************************************/
@@ -809,7 +776,7 @@ void UK2Node_CallFunction::CreateExecPinsForFunctionCall(const UFunction* Functi
 	if (!bIsPureFunc)
 	{
 		// If we want enum->exec expansion, and it is not disabled, do it now
-		if (bWantsEnumToExecExpansion)
+		if(bWantsEnumToExecExpansion)
 		{
 			TArray<FName> EnumNames;
 			GetExpandEnumPinNames(Function, EnumNames);
@@ -1193,17 +1160,15 @@ bool UK2Node_CallFunction::CreatePinsForFunctionCall(const UFunction* Function)
 	}
 
 	// If we have 'enum to exec' parameters, set their default value to something valid so we don't get warnings
-	if (bWantsEnumToExecExpansion)
+	if(bWantsEnumToExecExpansion)
 	{
-		using namespace UE::K2NodeCallFunction::Private;
-
 		TArray<FName> EnumNamesToCheck;
 		GetExpandEnumPinNames(Function, EnumNamesToCheck);
 
 		for (const FName& Name : EnumNamesToCheck)
 		{
-			UEdGraphPin* EnumParamPin = FindEnumParamPin(*this, Name);
-			if (UEnum* PinEnum = (EnumParamPin ? Cast<UEnum>(EnumParamPin->PinType.PinSubCategoryObject.Get()) : nullptr))
+			UEdGraphPin* EnumParamPin = FindPin(Name);
+			if (UEnum* PinEnum = (EnumParamPin ? Cast<UEnum>(EnumParamPin->PinType.PinSubCategoryObject.Get()) : NULL))
 			{
 				EnumParamPin->DefaultValue = PinEnum->GetNameStringByIndex(0);
 			}
@@ -2313,6 +2278,24 @@ void UK2Node_CallFunction::ValidateNodeDuringCompilation(class FCompilerResultsL
 			).ToString();
 			MessageLog.Error(*ErrorString, this);
 		}
+		else if (Pin && Pin->PinType.PinCategory == UEdGraphSchema_K2::PC_Object && Pin->PinName == UEdGraphSchema_K2::PN_Self)
+		{
+			const UEdGraphPin* SelfPin = MessageLog.FindSourcePin(Pin);
+			for (const UEdGraphPin* LinkedTo : SelfPin->LinkedTo)
+			{
+				if (ensure(LinkedTo) && LinkedTo->PinType.PinCategory == UEdGraphSchema_K2::PC_Interface)
+				{
+					if (LinkedTo->PinType.IsArray())
+					{
+						MessageLog.Note(*FText::Format(LOCTEXT("InterfaceArrayTargetConnectionNote", "@@: An array of interface types can no longer be directly connected to '{0}'. Each entry must first be cast to the object type. However, the existing connection to '{1}' will continue to work for backwards-compatibility."), FText::FromString(SelfPin->GetName()), FText::FromString(LinkedTo->GetName())).ToString(), this);
+					}
+					else
+					{
+						MessageLog.Note(*FText::Format(LOCTEXT("InterfaceTargetConnectionNote", "@@: An interface type can no longer be directly connected to '{0}'. It must first be cast to the object type. However, the existing connection to '{1}' will continue to work for backwards-compatibility."), FText::FromString(SelfPin->GetName()), FText::FromString(LinkedTo->GetName())).ToString(), this);
+					}
+				}
+			}
+		}
 	}
 }
 
@@ -2321,7 +2304,6 @@ void UK2Node_CallFunction::Serialize(FArchive& Ar)
 	Super::Serialize(Ar);
 
 	Ar.UsingCustomVersion(FReleaseObjectVersion::GUID);
-	Ar.UsingCustomVersion(FUE5MainStreamObjectVersion::GUID);
 
 	if (Ar.IsLoading())
 	{
@@ -2391,25 +2373,6 @@ void UK2Node_CallFunction::Serialize(FArchive& Ar)
 			}
 		}
 
-		for (int32 PinIndex = 0; PinIndex < Pins.Num(); ++PinIndex)
-		{
-			UEdGraphPin* Pin = Pins[PinIndex];
-			check(Pin);
-
-			bool bNeedsSubCategoryObjectRepair =
-				(Pin->PinType.PinCategory == UEdGraphSchema_K2::PC_Object) &&
-				(Pin->PinType.PinSubCategory != UEdGraphSchema_K2::PSC_Self) &&
-				(Pin->PinType.PinSubCategoryObject == nullptr) &&
-				(Ar.CustomVer(FReleaseObjectVersion::GUID) < FUE5MainStreamObjectVersion::NullPinSubCategoryObjectFix);
-
-			// Prior to NullPinSubCategoryObjectFix, some object pins were serialized with a null PinSubCategoryObject.
-			// Going forward, this will be an error, so we'll attempt to repair the pin by assigning a class.
-			if (bNeedsSubCategoryObjectRepair)
-			{
-				Pin->PinType.PinSubCategoryObject = FunctionReference.GetMemberParentClass();
-			}
-		}
-
 		if (!Ar.IsObjectReferenceCollector())
 		{
 			// Don't validate the enabled state if the user has explicitly set it. Also skip validation if we're just duplicating this node.
@@ -2502,12 +2465,10 @@ void UK2Node_CallFunction::ExpandNode(class FKismetCompilerContext& CompilerCont
 	}
 
 	// If we have an enum param that is expanded, we handle that first
-	if (bWantsEnumToExecExpansion)
+	if(bWantsEnumToExecExpansion)
 	{
-		if (Function)
+		if(Function)
 		{
-			using namespace UE::K2NodeCallFunction::Private;
-
 			TArray<FName> EnumNamesToCheck;
 			GetExpandEnumPinNames(Function, EnumNamesToCheck);
 
@@ -2556,26 +2517,20 @@ void UK2Node_CallFunction::ExpandNode(class FKismetCompilerContext& CompilerCont
 				}
 			};
 
-			for (const FName& ParamName : EnumNamesToCheck)
+			for (const FName& EnumParamName : EnumNamesToCheck)
 			{
 				UEnum* Enum = nullptr;
-				UEdGraphPin* EnumParamPin = nullptr;
-	
-				if (FBoolProperty* BoolProp = FindFProperty<FBoolProperty>(Function, ParamName))
-				{
-					EnumParamPin = FindBoolParamPin(*this, ParamName);
-				}
-				else if (FByteProperty* ByteProp = FindFProperty<FByteProperty>(Function, ParamName))
+
+				if (FByteProperty* ByteProp = FindFProperty<FByteProperty>(Function, EnumParamName))
 				{
 					Enum = ByteProp->Enum;
-					EnumParamPin = FindEnumParamPin(*this, ParamName);
 				}
-				else if (FEnumProperty* EnumProp = FindFProperty<FEnumProperty>(Function, ParamName))
+				else if (FEnumProperty* EnumProp = FindFProperty<FEnumProperty>(Function, EnumParamName))
 				{
 					Enum = EnumProp->GetEnum();
-					EnumParamPin = FindEnumParamPin(*this, ParamName);
 				}
 
+				UEdGraphPin* EnumParamPin = FindPin(EnumParamName);
 				if (Enum && EnumParamPin)
 				{
 					// Expanded as input execs pins
@@ -2812,18 +2767,34 @@ void UK2Node_CallFunction::ExpandNode(class FKismetCompilerContext& CompilerCont
 		}
 	}
 
+	// Older assets may have interface pins wired directly to the target pin in the source graph (due to an earlier regression).
+	UEdGraphPin* SelfPin = Schema->FindSelfPin(*this, EEdGraphPinDirection::EGPD_Input);
+	if (SelfPin && SelfPin->PinType.PinCategory == UEdGraphSchema_K2::PC_Object)
+	{
+		for (UEdGraphPin* PinLinkedToSelfPin : SelfPin->LinkedTo)
+		{
+			if (PinLinkedToSelfPin && PinLinkedToSelfPin->PinType.PinCategory == UEdGraphSchema_K2::PC_Interface && !PinLinkedToSelfPin->PinType.IsContainer())
+			{
+				PinLinkedToSelfPin->BreakLinkTo(SelfPin);
+				if (!Schema->TryCreateConnection(PinLinkedToSelfPin, SelfPin))
+				{
+					PinLinkedToSelfPin->MakeLinkTo(SelfPin);
+				}
+			}
+		}
+	}
+
 	// Then we go through and expand out array iteration if necessary
 	const bool bAllowMultipleSelfs = AllowMultipleSelfs(true);
-	UEdGraphPin* MultiSelf = Schema->FindSelfPin(*this, EEdGraphPinDirection::EGPD_Input);
-	if(bAllowMultipleSelfs && MultiSelf && !MultiSelf->PinType.IsArray())
+	if(bAllowMultipleSelfs && SelfPin && !SelfPin->PinType.IsArray())
 	{
 		const bool bProperInputToExpandForEach = 
-			(1 == MultiSelf->LinkedTo.Num()) && 
-			(nullptr != MultiSelf->LinkedTo[0]) && 
-			(MultiSelf->LinkedTo[0]->PinType.IsArray());
+			(1 == SelfPin->LinkedTo.Num()) &&
+			(nullptr != SelfPin->LinkedTo[0]) &&
+			(SelfPin->LinkedTo[0]->PinType.IsArray());
 		if(bProperInputToExpandForEach)
 		{
-			CallForEachElementInArrayExpansion(this, MultiSelf, CompilerContext, SourceGraph);
+			CallForEachElementInArrayExpansion(this, SelfPin, CompilerContext, SourceGraph);
 		}
 	}
 }
diff --git a/Engine/Source/Editor/BlueprintGraph/Private/K2Node_ClassDynamicCast.cpp b/Engine/Source/Editor/BlueprintGraph/Private/K2Node_ClassDynamicCast.cpp
--- a/Engine/Source/Editor/BlueprintGraph/Private/K2Node_ClassDynamicCast.cpp
+++ b/Engine/Source/Editor/BlueprintGraph/Private/K2Node_ClassDynamicCast.cpp
@@ -49,22 +49,8 @@ void UK2Node_ClassDynamicCast::AllocateDefaultPins()
 	//@TODO: Move this somewhere more sensible
 	ensure((TargetType == nullptr) || (!TargetType->HasAnyClassFlags(CLASS_NewerVersionExists)));
 
-	const UEdGraphSchema_K2* K2Schema = Cast<UEdGraphSchema_K2>(GetSchema());
-	check(K2Schema != nullptr);
-	if (!K2Schema->DoesGraphSupportImpureFunctions(GetGraph()))
-	{
-		bIsPureCast = true;
-	}
-
-	if (!bIsPureCast)
-	{
-		// Input - Execution Pin
-		CreatePin(EGPD_Input, UEdGraphSchema_K2::PC_Exec, UEdGraphSchema_K2::PN_Execute);
-
-		// Output - Execution Pins
-		CreatePin(EGPD_Output, UEdGraphSchema_K2::PC_Exec, UEdGraphSchema_K2::PN_CastSucceeded);
-		CreatePin(EGPD_Output, UEdGraphSchema_K2::PC_Exec, UEdGraphSchema_K2::PN_CastFailed);
-	}
+	// Exec pins (if needed)
+	CreateExecPins();
 
 	// Input - Source type Pin
 	CreatePin(EGPD_Input, UEdGraphSchema_K2::PC_Class, UObject::StaticClass(), FClassDynamicCastHelper::ClassToCastName);
@@ -76,8 +62,8 @@ void UK2Node_ClassDynamicCast::AllocateDefaultPins()
 		CreatePin(EGPD_Output, UEdGraphSchema_K2::PC_Class, *TargetType, *CastResultPinName);
 	}
 
-	UEdGraphPin* BoolSuccessPin = CreatePin(EGPD_Output, UEdGraphSchema_K2::PC_Boolean, FClassDynamicCastHelper::CastSuccessPinName);
-	BoolSuccessPin->bHidden = !bIsPureCast;
+	// Output - Success
+	CreateSuccessPin();
 
 	UK2Node::AllocateDefaultPins();
 }
diff --git a/Engine/Source/Editor/BlueprintGraph/Private/K2Node_DynamicCast.cpp b/Engine/Source/Editor/BlueprintGraph/Private/K2Node_DynamicCast.cpp
--- a/Engine/Source/Editor/BlueprintGraph/Private/K2Node_DynamicCast.cpp
+++ b/Engine/Source/Editor/BlueprintGraph/Private/K2Node_DynamicCast.cpp
@@ -30,6 +30,7 @@
 #include "ToolMenu.h"
 #include "ToolMenuSection.h"
 #include "UObject/Class.h"
+#include "UObject/FortniteMainBranchObjectVersion.h"
 #include "UObject/Interface.h"
 #include "UObject/NameTypes.h"
 #include "UObject/Object.h"
@@ -48,27 +49,62 @@ namespace UK2Node_DynamicCastImpl
 
 UK2Node_DynamicCast::UK2Node_DynamicCast(const FObjectInitializer& ObjectInitializer)
 	: Super(ObjectInitializer)
-	, bIsPureCast(false)
+	, PureState(EPureState::UseDefault)
 {
+PRAGMA_DISABLE_DEPRECATION_WARNINGS
+	// @todo_deprecated - Remove these later.
+	bIsPureCast = false;
+	bIsPureCast_DEPRECATED = bIsPureCast;
+PRAGMA_ENABLE_DEPRECATION_WARNINGS
 }
 
-void UK2Node_DynamicCast::AllocateDefaultPins()
+void UK2Node_DynamicCast::Serialize(FArchive& Ar)
 {
-	const bool bReferenceObsoleteClass = TargetType && TargetType->HasAnyClassFlags(CLASS_NewerVersionExists);
-	if (bReferenceObsoleteClass)
+	Super::Serialize(Ar);
+
+	Ar.UsingCustomVersion(FFortniteMainBranchObjectVersion::GUID);
+
+	if (Ar.IsLoading() && Ar.CustomVer(FFortniteMainBranchObjectVersion::GUID) < FFortniteMainBranchObjectVersion::DynamicCastNodesUsePureStateEnum)
 	{
-		Message_Error(FString::Printf(TEXT("Node '%s' references obsolete class '%s'"), *GetPathName(), *TargetType->GetPathName()));
+		if (bIsPureCast_DEPRECATED)
+		{
+			PureState = EPureState::Pure;
+		}
+		else
+		{
+			PureState = EPureState::Impure;
+		}
 	}
-	ensure(!bReferenceObsoleteClass);
+	else
+	{
+		Ar << PureState;
+	}
+}
 
+void UK2Node_DynamicCast::CreateExecPins()
+{
+	// Ensure the node is either pure or impure, based on current settings.
 	const UEdGraphSchema_K2* K2Schema = Cast<UEdGraphSchema_K2>(GetSchema());
 	check(K2Schema != nullptr);
 	if (!K2Schema->DoesGraphSupportImpureFunctions(GetGraph()))
 	{
-		bIsPureCast = true;
+		PureState = EPureState::Pure;
+	}
+	else if (PureState == EPureState::UseDefault)
+	{
+		const UBlueprintEditorSettings* BlueprintSettings = GetDefault<UBlueprintEditorSettings>();
+		if (BlueprintSettings->bFavorPureCastNodes)
+		{
+			PureState = EPureState::Pure;
+		}
+		else
+		{
+			PureState = EPureState::Impure;
+		}
 	}
 
-	if (!bIsPureCast)
+	const bool bIsNodePure = IsNodePure();
+	if (!bIsNodePure)
 	{
 		// Input - Execution Pin
 		CreatePin(EGPD_Input, UEdGraphSchema_K2::PC_Exec, UEdGraphSchema_K2::PN_Execute);
@@ -78,6 +114,30 @@ void UK2Node_DynamicCast::AllocateDefaultPins()
 		CreatePin(EGPD_Output, UEdGraphSchema_K2::PC_Exec, UEdGraphSchema_K2::PN_CastFailed);
 	}
 
+PRAGMA_DISABLE_DEPRECATION_WARNINGS
+	// @todo_deprecated - Remove this later.
+	bIsPureCast = bIsNodePure;
+PRAGMA_ENABLE_DEPRECATION_WARNINGS
+}
+
+void UK2Node_DynamicCast::CreateSuccessPin()
+{
+	UEdGraphPin* BoolSuccessPin = CreatePin(EGPD_Output, UEdGraphSchema_K2::PC_Boolean, UK2Node_DynamicCastImpl::CastSuccessPinName);
+	BoolSuccessPin->bHidden = !IsNodePure();
+}
+
+void UK2Node_DynamicCast::AllocateDefaultPins()
+{
+	const bool bReferenceObsoleteClass = TargetType && TargetType->HasAnyClassFlags(CLASS_NewerVersionExists);
+	if (bReferenceObsoleteClass)
+	{
+		Message_Error(FString::Printf(TEXT("Node '%s' references obsolete class '%s'"), *GetPathName(), *TargetType->GetPathName()));
+	}
+	ensure(!bReferenceObsoleteClass);
+
+	// Exec pins (if needed)
+	CreateExecPins();
+
 	// Input - Source type Pin
 	CreatePin(EGPD_Input, UEdGraphSchema_K2::PC_Wildcard, UObject::StaticClass(), UEdGraphSchema_K2::PN_ObjectToCast);
 
@@ -95,8 +155,8 @@ void UK2Node_DynamicCast::AllocateDefaultPins()
 		}
 	}
 
-	UEdGraphPin* BoolSuccessPin = CreatePin(EGPD_Output, UEdGraphSchema_K2::PC_Boolean, UK2Node_DynamicCastImpl::CastSuccessPinName);
-	BoolSuccessPin->bHidden = !bIsPureCast;
+	// Output - Success
+	CreateSuccessPin();
 
 	Super::AllocateDefaultPins();
 }
@@ -175,7 +235,7 @@ void UK2Node_DynamicCast::GetNodeContextMenuActions(UToolMenu* Menu, UGraphNodeC
 				return bInCanTogglePurity;
 			};
 
-			if (bIsPureCast)
+			if (IsNodePure())
 			{
 				MenuEntryTitle = LOCTEXT("MakeImpureTitle", "Convert to impure cast");
 				MenuEntryTooltip = LOCTEXT("MakeImpureTooltip", "Adds in branching execution pins so that you can separatly handle when the cast fails/succeeds.");
@@ -213,18 +273,10 @@ void UK2Node_DynamicCast::PostReconstructNode()
 	NotifyPinConnectionListChanged(GetCastSourcePin());
 }
 
-void UK2Node_DynamicCast::PostPlacedNewNode()
-{
-	Super::PostPlacedNewNode();
-
-	const UBlueprintEditorSettings* BlueprintSettings = GetDefault<UBlueprintEditorSettings>();
-	SetPurity(BlueprintSettings->bFavorPureCastNodes);
-}
-
 UEdGraphPin* UK2Node_DynamicCast::GetValidCastPin() const
 {
 	UEdGraphPin* Pin = FindPin(UEdGraphSchema_K2::PN_CastSucceeded);
-	check((Pin != nullptr) || bIsPureCast);
+	check((Pin != nullptr) || IsNodePure());
 	check((Pin == nullptr) || (Pin->Direction == EGPD_Output));
 	return Pin;
 }
@@ -232,7 +284,7 @@ UEdGraphPin* UK2Node_DynamicCast::GetValidCastPin() const
 UEdGraphPin* UK2Node_DynamicCast::GetInvalidCastPin() const
 {
 	UEdGraphPin* Pin = FindPin(UEdGraphSchema_K2::PN_CastFailed);
-	check((Pin != nullptr) || bIsPureCast);
+	check((Pin != nullptr) || IsNodePure());
 	check((Pin == nullptr) || (Pin->Direction == EGPD_Output));
 	return Pin;
 }
@@ -272,9 +324,21 @@ UEdGraphPin* UK2Node_DynamicCast::GetBoolSuccessPin() const
 
 void UK2Node_DynamicCast::SetPurity(bool bNewPurity)
 {
-	if (bNewPurity != bIsPureCast)
+	if (bNewPurity != IsNodePure())
 	{
+		if (bNewPurity)
+		{
+			PureState = EPureState::Pure;
+		}
+		else
+		{
+			PureState = EPureState::Impure;
+		}
+
+PRAGMA_DISABLE_DEPRECATION_WARNINGS
+		// @todo_deprecated - Remove this later.
 		bIsPureCast = bNewPurity;
+PRAGMA_ENABLE_DEPRECATION_WARNINGS
 
 		bool const bHasBeenConstructed = (Pins.Num() > 0);
 		if (bHasBeenConstructed)
@@ -286,11 +350,12 @@ void UK2Node_DynamicCast::SetPurity(bool bNewPurity)
 
 void UK2Node_DynamicCast::TogglePurity()
 {
-	const FText TransactionTitle = bIsPureCast ? LOCTEXT("TogglePurityToImpure", "Convert to Impure Cast") : LOCTEXT("TogglePurityToPure", "Convert to Pure Cast");
+	const bool bIsNodePure = IsNodePure();
+	const FText TransactionTitle = bIsNodePure ? LOCTEXT("TogglePurityToImpure", "Convert to Impure Cast") : LOCTEXT("TogglePurityToPure", "Convert to Pure Cast");
 	const FScopedTransaction Transaction( TransactionTitle );
 	Modify();
 
-	SetPurity(!bIsPureCast);
+	SetPurity(!bIsNodePure);
 }
 
 UK2Node::ERedirectType UK2Node_DynamicCast::DoPinsMatchForReconstruction(const UEdGraphPin* NewPin, int32 NewPinIndex, const UEdGraphPin* OldPin, int32 OldPinIndex) const
@@ -496,7 +561,7 @@ void UK2Node_DynamicCast::ValidateNodeDuringCompilation(FCompilerResultsLog& Mes
 
 bool UK2Node_DynamicCast::ReconnectPureExecPins(TArray<UEdGraphPin*>& OldPins)
 {
-	if (bIsPureCast)
+	if (IsNodePure())
 	{
 		// look for an old exec pin
 		UEdGraphPin* PinExec = nullptr;
diff --git a/Engine/Source/Editor/Kismet/Private/BlueprintCompilationManager.cpp b/Engine/Source/Editor/Kismet/Private/BlueprintCompilationManager.cpp
--- a/Engine/Source/Editor/Kismet/Private/BlueprintCompilationManager.cpp
+++ b/Engine/Source/Editor/Kismet/Private/BlueprintCompilationManager.cpp
@@ -2731,6 +2731,11 @@ UClass* FBlueprintCompilationManagerImpl::FastGenerateSkeletonClass(UBlueprint*
 		}
 
 		UFunction* ParentFn = ParentClass->FindFunctionByName(NewFunction->GetFName());
+
+		// Set the parent function prior to searching for a corresponding interface class function. This matches what
+		// is done in the full path (@see FKismetCompilerContext::PrecompileFunction). It is ok for this to be NULL.
+		NewFunction->SetSuperStruct(ParentFn);
+
 		if(ParentFn == nullptr)
 		{
 			// check for function in implemented interfaces:
@@ -2758,7 +2763,6 @@ UClass* FBlueprintCompilationManagerImpl::FastGenerateSkeletonClass(UBlueprint*
 				}
 			}
 		}
-		NewFunction->SetSuperStruct( ParentFn );
 		
 		InCurrentParamStorageLocation = &NewFunction->ChildProperties;
 
diff --git a/Engine/Source/Runtime/Core/Public/UObject/FortniteMainBranchObjectVersion.h b/Engine/Source/Runtime/Core/Public/UObject/FortniteMainBranchObjectVersion.h
--- a/Engine/Source/Runtime/Core/Public/UObject/FortniteMainBranchObjectVersion.h
+++ b/Engine/Source/Runtime/Core/Public/UObject/FortniteMainBranchObjectVersion.h
@@ -265,6 +265,9 @@ struct CORE_API FFortniteMainBranchObjectVersion
 		// Change HLODActors to RuntimeCells mapping to use a GUID instead of the cell name
 		WorldPartitionHLODActorUseSourceCellGuid,
 
+		// Dynamic cast nodes use an enumerated pure node state to include a value for the default setting.
+		DynamicCastNodesUsePureStateEnum,
+
 		// -----<new versions can be added above this line>-------------------------------------------------
 		VersionPlusOne,
 		LatestVersion = VersionPlusOne - 1
