Fix for an ASAN trigger source in compiled Blueprint graphs with direct interface->self pin connections on call site nodes for implemented interface functions.

Note: This is low-risk patch in lieu of a more in-depth fix that will target a future release. This change specifically patches the critical runtime stack issue without breaking backwards-compatibility.

diff --git a/Engine/Source/Editor/BlueprintGraph/Private/CallFunctionHandler.cpp b/Engine/Source/Editor/BlueprintGraph/Private/CallFunctionHandler.cpp
--- a/Engine/Source/Editor/BlueprintGraph/Private/CallFunctionHandler.cpp
+++ b/Engine/Source/Editor/BlueprintGraph/Private/CallFunctionHandler.cpp
@@ -468,11 +468,18 @@ void FKCHandler_CallFunction::CreateFunctionCallStatement(FKismetFunctionContext
 			FBlueprintCompiledStatement* LatentStatement = nullptr;
 			for (FBPTerminal* Target : ContextTerms)
 			{
+				// Currently, call site nodes will (incorrectly) expose the target pin as an interface type for calls to
+				// interface functions that are implemented by the owning class, so in that case we need to flag that the
+				// calling context is an interface if the target pin is also linked to an interface pin type (e.g. result
+				// of a cast node). Otherwise, we'll infer the wrong context type at runtime and corrupt the stack by
+				// reading an interface ptr (16 bytes) into an object ptr (8 bytes) when we process the context opcode.
+				const bool bIsInterfaceContextTerm = Target && Target->AssociatedVarProperty && Target->AssociatedVarProperty->IsA<FInterfaceProperty>();
+
 				FBlueprintCompiledStatement& Statement = Context.AppendStatementForNode(Node);
 				Statement.FunctionToCall = Function;
 				Statement.FunctionContext = Target;
 				Statement.Type = KCST_CallFunction;
-				Statement.bIsInterfaceContext = IsCalledFunctionFromInterface(Node);
+				Statement.bIsInterfaceContext = IsCalledFunctionFromInterface(Node) || bIsInterfaceContextTerm;
 				Statement.bIsParentContext = IsCalledFunctionFinal(Node);
 
 				Statement.LHS = LHSTerm;
