Fix to include specialized auto-convert cases in addition to auto-cast when matching pin types with potential redirects during node reconstruction.

Notes:
- UEdGraphSchema_K2::CreateAutomaticConversionNodeAndConnections() handles both, but we were only allowing auto-cast scenarios to pass in the revised UK2Node::DoPinsMatchForReconstruction().
- This is needed for linked pins with matching names which were previously being allowed through, but now require types to match as well (or at least have a matching auto-conversion or auto-cast operation).
- UK2Node::DoPinsMatchForReconstruction() is now at parity with UEdGraphSchema_K2::CanCreateConnection() in terms of when the latter returns CONNECT_RESPONSE_MAKE_WITH_CONVERSION_NODE.

#rnx

diff --git a/Engine/Source/Editor/BlueprintGraph/Private/K2Node.cpp b/Engine/Source/Editor/BlueprintGraph/Private/K2Node.cpp
--- a/Engine/Source/Editor/BlueprintGraph/Private/K2Node.cpp
+++ b/Engine/Source/Editor/BlueprintGraph/Private/K2Node.cpp
@@ -911,7 +911,8 @@ UK2Node::ERedirectType UK2Node::DoPinsMatchForReconstruction(const UEdGraphPin*
 			}
 			else 
 			{
-				if (K2Schema->SearchForAutocastFunction(OldPin->PinType, NewPin->PinType).IsSet())
+				if (K2Schema->SearchForAutocastFunction(OldPin->PinType, NewPin->PinType).IsSet()
+					|| K2Schema->FindSpecializedConversionNode(OldPin->PinType, *NewPin, /* bCreateNode = */false).IsSet())
 				{
 					RedirectType = static_cast<ERedirectType>(static_cast<uint8>(ERedirectType_Name) | static_cast<uint8>(ERedirectType_Type));
 				}
@@ -1000,7 +1001,8 @@ UK2Node::ERedirectType UK2Node::DoPinsMatchForReconstruction(const UEdGraphPin*
 			// A type change might've also occurred with a rename.
 			if (RedirectType == ERedirectType_Name)
 			{
-				if (K2Schema->SearchForAutocastFunction(OldPin->PinType, NewPin->PinType).IsSet())
+				if (K2Schema->SearchForAutocastFunction(OldPin->PinType, NewPin->PinType).IsSet()
+					|| K2Schema->FindSpecializedConversionNode(OldPin->PinType, *NewPin, /* bCreateNode = */false).IsSet())
 				{
 					RedirectType = static_cast<ERedirectType>(static_cast<uint8>(ERedirectType_Name) | static_cast<uint8>(ERedirectType_Type));
 				}
