FBPTerminal::ObjectLiteral to TObjectPtr from UObject to avoid loading object during Bluepirnt compilation
https://p4-swarm.epicgames.net/reviews/23978331

diff --git a/Engine/Source/Editor/BlueprintGraph/Classes/EdGraphSchema_K2.h b/Engine/Source/Editor/BlueprintGraph/Classes/EdGraphSchema_K2.h
--- a/Engine/Source/Editor/BlueprintGraph/Classes/EdGraphSchema_K2.h
+++ b/Engine/Source/Editor/BlueprintGraph/Classes/EdGraphSchema_K2.h
@@ -530,7 +530,7 @@ public:
 	virtual bool TryCreateConnection(UEdGraphPin* A, UEdGraphPin* B) const override;
 	virtual bool CreateAutomaticConversionNodeAndConnections(UEdGraphPin* A, UEdGraphPin* B) const override;
 	virtual bool CreatePromotedConnection(UEdGraphPin* A, UEdGraphPin* B) const override;
-	virtual FString IsPinDefaultValid(const UEdGraphPin* Pin, const FString& NewDefaultValue, UObject* NewDefaultObject, const FText& InNewDefaultText) const override;
+	virtual FString IsPinDefaultValid(const UEdGraphPin* Pin, const FString& NewDefaultValue, TObjectPtr<UObject> NewDefaultObject, const FText& InNewDefaultText) const override;
 	virtual bool DoesSupportPinWatching() const	override;
 	virtual bool IsPinBeingWatched(UEdGraphPin const* Pin) const override;
 	virtual void ClearPinWatch(UEdGraphPin const* Pin) const override;
@@ -670,7 +670,7 @@ public:
 	virtual void GetPinDefaultValuesFromString(const FEdGraphPinType& PinType, UObject* OwningObject, const FString& NewValue, FString& UseDefaultValue, TObjectPtr<UObject>& UseDefaultObject, FText& UseDefaultText, bool bPreserveTextIdentity = true) const;
 
 	/** Do validation, that doesn't require a knowledge about actual pin */
-	virtual bool DefaultValueSimpleValidation(const FEdGraphPinType& PinType, const FName PinName, const FString& NewDefaultValue, UObject* NewDefaultObject, const FText& InText, FString* OutMsg = nullptr) const;
+	virtual bool DefaultValueSimpleValidation(const FEdGraphPinType& PinType, const FName PinName, const FString& NewDefaultValue, TObjectPtr<UObject> NewDefaultObject, const FText& InText, FString* OutMsg = nullptr) const;
 
 	/** Returns true if the owning node is a function with AutoCreateRefTerm meta data */
 	static bool IsAutoCreateRefTerm(const UEdGraphPin* Pin);
diff --git a/Engine/Source/Editor/BlueprintGraph/Private/EdGraphSchema_K2.cpp b/Engine/Source/Editor/BlueprintGraph/Private/EdGraphSchema_K2.cpp
--- a/Engine/Source/Editor/BlueprintGraph/Private/EdGraphSchema_K2.cpp
+++ b/Engine/Source/Editor/BlueprintGraph/Private/EdGraphSchema_K2.cpp
@@ -2963,7 +2963,7 @@ bool UEdGraphSchema_K2::CreatePromotedConnection(UEdGraphPin* PinA, UEdGraphPin*
 	return true;
 }
 
-FString UEdGraphSchema_K2::IsPinDefaultValid(const UEdGraphPin* Pin, const FString& NewDefaultValue, UObject* NewDefaultObject, const FText& InNewDefaultText) const
+FString UEdGraphSchema_K2::IsPinDefaultValid(const UEdGraphPin* Pin, const FString& NewDefaultValue, TObjectPtr<UObject> NewDefaultObject, const FText& InNewDefaultText) const
 {
 	check(Pin);
 
@@ -4259,7 +4259,7 @@ namespace
 };
 
 
-bool UEdGraphSchema_K2::DefaultValueSimpleValidation(const FEdGraphPinType& PinType, const FName PinName, const FString& NewDefaultValue, UObject* NewDefaultObject, const FText& InNewDefaultText, FString* OutMsg /*= NULL*/) const
+bool UEdGraphSchema_K2::DefaultValueSimpleValidation(const FEdGraphPinType& PinType, const FName PinName, const FString& NewDefaultValue, TObjectPtr<UObject> NewDefaultObject, const FText& InNewDefaultText, FString* OutMsg /*= NULL*/) const
 {
 #ifdef DVSV_RETURN_MSG
 	static_assert(false, "Macro redefinition.");
@@ -4404,14 +4404,14 @@ bool UEdGraphSchema_K2::DefaultValueSimpleValidation(const FEdGraphPinType& PinT
 		}
 		if (NewDefaultObject != nullptr && ObjectClass != nullptr)
 		{
-			const UClass* AuthoritativeClass = NewDefaultObject->GetClass()->GetAuthoritativeClass();
+			const UClass* AuthoritativeClass = NewDefaultObject.GetClass()->GetAuthoritativeClass();
 			if (!AuthoritativeClass || !AuthoritativeClass->IsChildOf(ObjectClass))
 			{
 				// Not a type of object, but is it an object implementing an interface?
-                if(PinCategory != PC_Interface || !NewDefaultObject->GetClass()->ImplementsInterface(ObjectClass))
-                {
-                	DVSV_RETURN_MSG(FString::Printf(TEXT("%s isn't a %s (specified on pin %s)"), *NewDefaultObject->GetPathName(), *ObjectClass->GetName(), *PinName.ToString()));
-                }
+				if(PinCategory != PC_Interface || !NewDefaultObject.GetClass()->ImplementsInterface(ObjectClass))
+				{
+					DVSV_RETURN_MSG(FString::Printf(TEXT("%s isn't a %s (specified on pin %s)"), *NewDefaultObject->GetPathName(), *ObjectClass->GetName(), *PinName.ToString()));
+				}
 			}
 		}
 	}
diff --git a/Engine/Source/Editor/BlueprintGraph/Private/K2Node_GetClassDefaults.cpp b/Engine/Source/Editor/BlueprintGraph/Private/K2Node_GetClassDefaults.cpp
--- a/Engine/Source/Editor/BlueprintGraph/Private/K2Node_GetClassDefaults.cpp
+++ b/Engine/Source/Editor/BlueprintGraph/Private/K2Node_GetClassDefaults.cpp
@@ -163,7 +163,7 @@ namespace
 						ClassContextTerm->SetContextTypeClass();
 
 						// Infer the class type from the context term
-						if(const UClass* ClassType = Cast<UClass>(ClassContextTerm->bIsLiteral ? ClassContextTerm->ObjectLiteral : ClassContextTerm->Type.PinSubCategoryObject.Get()))
+						if(const UClass* ClassType = Cast<UClass>(ClassContextTerm->bIsLiteral ? ClassContextTerm->ObjectLiteral.Get() : ClassContextTerm->Type.PinSubCategoryObject.Get()))
 						{
 							// Create a local term for each output pin (class property)
 							for(int32 PinIndex = 0; PinIndex < Node->Pins.Num(); ++PinIndex)
diff --git a/Engine/Source/Editor/KismetCompiler/Private/KismetCompilerVMBackend.cpp b/Engine/Source/Editor/KismetCompiler/Private/KismetCompilerVMBackend.cpp
--- a/Engine/Source/Editor/KismetCompiler/Private/KismetCompilerVMBackend.cpp
+++ b/Engine/Source/Editor/KismetCompiler/Private/KismetCompilerVMBackend.cpp
@@ -92,6 +92,14 @@ public:
 		return Ar;
 	}
 
+	FArchive& operator<<(FObjectPtr& Res) override
+	{
+		ScriptPointerType D = (ScriptPointerType)Res.GetHandle().PointerOrRef;
+		FArchive& Ar = *this;
+
+		Ar << D;
+		return Ar;
+	}
 	FArchive& operator<<(FField*& Res) override
 	{
 		ScriptPointerType D = (ScriptPointerType)Res;
@@ -1051,7 +1059,7 @@ public:
 				{
 					Writer << EX_Self;
 				}
-				else if (Term->ObjectLiteral == NULL)
+				else if (!Term->ObjectLiteral)
 				{
 					Writer << EX_NoObject;
 				}
diff --git a/Engine/Source/Editor/KismetCompiler/Public/BPTerminal.h b/Engine/Source/Editor/KismetCompiler/Public/BPTerminal.h
--- a/Engine/Source/Editor/KismetCompiler/Public/BPTerminal.h
+++ b/Engine/Source/Editor/KismetCompiler/Public/BPTerminal.h
@@ -34,7 +34,7 @@ struct FBPTerminal
 	FProperty* AssociatedVarProperty;
 
 	/** Pointer to an object literal */
-	UObject* ObjectLiteral;
+	TObjectPtr<UObject> ObjectLiteral;
 
 	/** The FText literal */
 	FText TextLiteral;
diff --git a/Engine/Source/Runtime/CoreUObject/Private/UObject/Class.cpp b/Engine/Source/Runtime/CoreUObject/Private/UObject/Class.cpp
--- a/Engine/Source/Runtime/CoreUObject/Private/UObject/Class.cpp
+++ b/Engine/Source/Runtime/CoreUObject/Private/UObject/Class.cpp
@@ -2129,7 +2129,13 @@ const UStruct* UStruct::HasMetaDataHierarchical(const FName& Key) const
 	{
 		ScriptPointerType  Temp = FPlatformMemory::ReadUnaligned<ScriptPointerType>(ScriptPtr);
 		UObject*& ExprPtrRef = (UObject*&)Temp;
-		if (ULinkerPlaceholderClass* PlaceholderObj = Cast<ULinkerPlaceholderClass>(ExprPtrRef))
+
+		TObjectPtr<UObject> ObjPtr(ExprPtrRef);
+		if (!ObjPtr.IsResolved())
+		{
+			return;
+		}
+		else if (ULinkerPlaceholderClass* PlaceholderObj = Cast<ULinkerPlaceholderClass>(ExprPtrRef))
 		{
 			PlaceholderObj->AddReferencingScriptExpr((UClass**)(&ExprPtrRef));
 		}
diff --git a/Engine/Source/Runtime/CoreUObject/Public/Serialization/ArchiveScriptReferenceCollector.h b/Engine/Source/Runtime/CoreUObject/Public/Serialization/ArchiveScriptReferenceCollector.h
--- a/Engine/Source/Runtime/CoreUObject/Public/Serialization/ArchiveScriptReferenceCollector.h
+++ b/Engine/Source/Runtime/CoreUObject/Public/Serialization/ArchiveScriptReferenceCollector.h
@@ -64,7 +64,7 @@ protected:
 	* @param Object	reference to Object reference
 	* @return reference to instance of this class
 	*/
-	FArchive& operator<<(UObject*& Object)
+	virtual FArchive& operator<<(UObject*& Object) override
 	{
 		// Avoid duplicate entries.
 		if (Object != nullptr && !ObjectArray.Contains(Object))
@@ -75,13 +75,24 @@ protected:
 		return *this;
 	}
 
+	virtual FArchive& operator<<(FObjectPtr& Object) override
+	{
+		if (Object.IsResolved())
+		{
+			UObject* Obj = Object.Get();
+			FArchive& Ar = *this << Obj;
+			return Ar;
+		}
+		return *this;
+	}
+
 	/**
 	* FField serialize operator implementation
 	*
 	* @param Field	reference to field reference
 	* @return reference to instance of this class
 	*/
-	FArchive& operator<<(FField*& Field)
+	virtual FArchive& operator<<(FField*& Field) override
 	{
 		if (Field != nullptr)
 		{
diff --git a/Engine/Source/Runtime/CoreUObject/Public/UObject/ScriptSerialization.h b/Engine/Source/Runtime/CoreUObject/Public/UObject/ScriptSerialization.h
--- a/Engine/Source/Runtime/CoreUObject/Public/UObject/ScriptSerialization.h
+++ b/Engine/Source/Runtime/CoreUObject/Public/UObject/ScriptSerialization.h
@@ -124,6 +124,22 @@
 	}
 #endif	//	XFERPTR
 
+#ifndef XFERTOBJPTR
+	#define XFERTOBJPTR() \
+	{ \
+		TObjectPtr<UObject> AlignedPtr; \
+		if (!Ar.IsLoading()) \
+		{ \
+			FMemory::Memcpy(&AlignedPtr, &Script[iCode], sizeof(ScriptPointerType)); \
+		} \
+			Ar << AlignedPtr; \
+		if (!Ar.IsSaving()) \
+		{ \
+			FMemory::Memcpy(&Script[iCode], &AlignedPtr, sizeof(ScriptPointerType)); \
+		} \
+		iCode += sizeof(ScriptPointerType); \
+	}
+#endif	//	XFERTOBJPTR
 
 #ifndef XFER_FUNC_POINTER
 	#define XFER_FUNC_POINTER	XFERPTR(UStruct*)
@@ -141,6 +157,10 @@
 	#define XFER_OBJECT_POINTER(Type)	XFERPTR(Type)
 #endif
 
+#ifndef XFER_TOBJECT_PTR
+	#define XFER_TOBJECT_PTR	XFERTOBJPTR
+#endif
+
 #ifndef FIXUP_EXPR_OBJECT_POINTER
 	// sometimes after a UOBject* expression is loaded it may require some post-
 	// processing (see: the overridden FIXUP_EXPR_OBJECT_POINTER(), defined in Class.cpp)
@@ -376,8 +396,8 @@
 		}
 		case EX_ObjectConst:
 		{
-			XFER_OBJECT_POINTER(UObject*);
-			FIXUP_EXPR_OBJECT_POINTER(UObject*);
+			XFER_TOBJECT_PTR();
+			FIXUP_EXPR_OBJECT_POINTER(TObjectPtr<UObject>);
 
 			break;
 		}
diff --git a/Engine/Source/Runtime/CoreUObject/Public/UObject/Stack.h b/Engine/Source/Runtime/CoreUObject/Public/UObject/Stack.h
--- a/Engine/Source/Runtime/CoreUObject/Public/UObject/Stack.h
+++ b/Engine/Source/Runtime/CoreUObject/Public/UObject/Stack.h
@@ -348,7 +348,12 @@ inline UObject* FFrame::ReadObject()
 	// turn that uint32 into a UObject pointer
 	UObject* Result = (UObject*)(TempCode);
 	Code += sizeof(ScriptPointerType);
+#if UE_WITH_OBJECT_HANDLE_LATE_RESOLVE
+	TObjectPtr<UObject> ObjPtr(Result);
+	return ObjPtr.Get();
+#else
 	return Result;
+#endif
 }
 
 inline FProperty* FFrame::ReadProperty()
diff --git a/Engine/Source/Runtime/Engine/Classes/EdGraph/EdGraphSchema.h b/Engine/Source/Runtime/Engine/Classes/EdGraph/EdGraphSchema.h
--- a/Engine/Source/Runtime/Engine/Classes/EdGraph/EdGraphSchema.h
+++ b/Engine/Source/Runtime/Engine/Classes/EdGraph/EdGraphSchema.h
@@ -820,7 +820,7 @@ class ENGINE_API UEdGraphSchema : public UObject
 	 *
 	 * @return	An empty string if the new value is legal, otherwise a message describing why it is invalid.
 	 */
-	virtual FString IsPinDefaultValid(const UEdGraphPin* Pin, const FString& NewDefaultValue, UObject* NewDefaultObject, const FText& InNewDefaultText) const  { return TEXT("Not implemented by this schema"); }
+	virtual FString IsPinDefaultValid(const UEdGraphPin* Pin, const FString& NewDefaultValue, TObjectPtr<UObject> NewDefaultObject, const FText& InNewDefaultText) const  { return TEXT("Not implemented by this schema"); }
 
 	/** 
 	 *	Determine whether the current pin default values are valid
