Allow self pins that redirect from interface->object type to pass without explicitly inserting an auto-conversion node during reconstruction of function call nodes.

Notes:
- This change is part of an upcoming regression fix (UE-157527) that will merge in from another stream that doesn't yet include type redirect support (UE-172222).

#rnx

diff --git a/Engine/Source/Editor/BlueprintGraph/Classes/K2Node_CallFunction.h b/Engine/Source/Editor/BlueprintGraph/Classes/K2Node_CallFunction.h
--- a/Engine/Source/Editor/BlueprintGraph/Classes/K2Node_CallFunction.h
+++ b/Engine/Source/Editor/BlueprintGraph/Classes/K2Node_CallFunction.h
@@ -119,6 +119,7 @@ public:
 	// End of UEdGraphNode interface
 
 	// UK2Node interface
+	virtual ERedirectType DoPinsMatchForReconstruction(const UEdGraphPin* NewPin, int32 NewPinIndex, const UEdGraphPin* OldPin, int32 OldPinIndex) const override;
 	virtual void ReallocatePinsDuringReconstruction(TArray<UEdGraphPin*>& OldPins) override;
 	virtual bool IsNodePure() const override { return bIsPureFunc; }
 	virtual void PostReconstructNode() override;
diff --git a/Engine/Source/Editor/BlueprintGraph/Private/K2Node_CallFunction.cpp b/Engine/Source/Editor/BlueprintGraph/Private/K2Node_CallFunction.cpp
--- a/Engine/Source/Editor/BlueprintGraph/Private/K2Node_CallFunction.cpp
+++ b/Engine/Source/Editor/BlueprintGraph/Private/K2Node_CallFunction.cpp
@@ -767,6 +767,27 @@ UEdGraphPin* FindSelfPin(TArray<UEdGraphPin*>& Pins)
 	return NULL;
 }
 
+UK2Node::ERedirectType UK2Node_CallFunction::DoPinsMatchForReconstruction(const UEdGraphPin* NewPin, int32 NewPinIndex, const UEdGraphPin* OldPin, int32 OldPinIndex) const
+{
+	ERedirectType Result = Super::DoPinsMatchForReconstruction(NewPin, NewPinIndex, OldPin, OldPinIndex);
+	const bool bNeedsAutoConvert = (Result & ERedirectType_Type);
+	if (bNeedsAutoConvert)
+	{
+		// Function call nodes implicitly handle interface->object type redirects on self pins.
+		const UEdGraphSchema_K2* K2Schema = CastChecked<const UEdGraphSchema_K2>(GetSchema());
+		if (K2Schema->IsSelfPin(*OldPin)
+			&& OldPin->PinType.PinCategory == UEdGraphSchema_K2::PC_Interface
+			&& NewPin->PinType.PinCategory == UEdGraphSchema_K2::PC_Object)
+		{
+			// Remove the conversion flag as we aren't going to explicitly add one here; instead,
+			// we add it into the intermediate graph at expansion time (for backwards-compatibility).
+			Result = static_cast<ERedirectType>(static_cast<uint8>(Result) & ~static_cast<uint8>(ERedirectType_Type));
+		}
+	}
+
+	return Result;
+}
+
 void UK2Node_CallFunction::ReallocatePinsDuringReconstruction(TArray<UEdGraphPin*>& OldPins)
 {
 	// BEGIN TEMP
