Exported K2 map and container nodes to allow out-of-module inheritance

diff --git a/Engine/Source/Editor/BlueprintGraph/Classes/K2Node_MakeContainer.h b/Engine/Source/Editor/BlueprintGraph/Classes/K2Node_MakeContainer.h
--- a/Engine/Source/Editor/BlueprintGraph/Classes/K2Node_MakeContainer.h
+++ b/Engine/Source/Editor/BlueprintGraph/Classes/K2Node_MakeContainer.h
@@ -43,8 +43,8 @@ protected:
 };
 
 
-UCLASS(MinimalAPI, abstract)
-class UK2Node_MakeContainer : public UK2Node, public IK2Node_AddPinInterface
+UCLASS(abstract)
+class BLUEPRINTGRAPH_API UK2Node_MakeContainer : public UK2Node, public IK2Node_AddPinInterface
 {
 	GENERATED_UCLASS_BODY()
 
@@ -53,9 +53,9 @@ class UK2Node_MakeContainer : public UK2Node, public IK2Node_AddPinInterface
 	int32 NumInputs;
 
 public:
-	BLUEPRINTGRAPH_API void RemoveInputPin(UEdGraphPin* Pin);
+	void RemoveInputPin(UEdGraphPin* Pin);
 
-	BLUEPRINTGRAPH_API UEdGraphPin* GetOutputPin() const;
+	UEdGraphPin* GetOutputPin() const;
 
 	/** returns a reference to the output array pin of this node, which is responsible for defining the type */
 	virtual FName GetOutputPinName() const PURE_VIRTUAL(UK2Node_MakeContainer::GetOutputPinName, return NAME_None;);
@@ -79,7 +79,7 @@ public:
 	// End of UK2Node interface
 
 	// IK2Node_AddPinInterface interface
-	BLUEPRINTGRAPH_API virtual void AddInputPin() override;
+	virtual void AddInputPin() override;
 	// End of IK2Node_AddPinInterface interface
 
 protected:
diff --git a/Engine/Source/Editor/BlueprintGraph/Classes/K2Node_MakeMap.h b/Engine/Source/Editor/BlueprintGraph/Classes/K2Node_MakeMap.h
--- a/Engine/Source/Editor/BlueprintGraph/Classes/K2Node_MakeMap.h
+++ b/Engine/Source/Editor/BlueprintGraph/Classes/K2Node_MakeMap.h
@@ -19,8 +19,8 @@ class UEdGraphPin;
 class UObject;
 struct FLinearColor;
 
-UCLASS(MinimalAPI)
-class UK2Node_MakeMap : public UK2Node_MakeContainer
+UCLASS()
+class BLUEPRINTGRAPH_API UK2Node_MakeMap : public UK2Node_MakeContainer
 {
 	GENERATED_UCLASS_BODY()
 
@@ -47,6 +47,6 @@ public:
 	// UK2Node_MakeContainer interface
 
 	// IK2Node_AddPinInterface interface
-	BLUEPRINTGRAPH_API virtual void AddInputPin() override;
+	virtual void AddInputPin() override;
 	// End of IK2Node_AddPinInterface interface
 };
