Generated Blueprint skeleton class functions will no longer parent themselves to interface class functions in the fast path, while also ensuring backwards-compatibility for existing call sites to implemented interface functions.

Notes:
- See 24252017 for additional context/review notes. The original change was backed out due to not handling mismatched object->object pin contexts on existing function call nodes (rare), so this version adds that part. An A/B screenshot is attached to the latest swarm review.

#rnx

diff --git a/Engine/Source/Editor/BlueprintGraph/Classes/K2Node_DynamicCast.h b/Engine/Source/Editor/BlueprintGraph/Classes/K2Node_DynamicCast.h
--- a/Engine/Source/Editor/BlueprintGraph/Classes/K2Node_DynamicCast.h
+++ b/Engine/Source/Editor/BlueprintGraph/Classes/K2Node_DynamicCast.h
@@ -34,6 +34,10 @@ class UK2Node_DynamicCast : public UK2Node
 	UPROPERTY()
 	TSubclassOf<class UObject>  TargetType;
 
+	//~ Begin UObject Interface
+	virtual void Serialize(FArchive& Ar);
+	//~ End UObject Interface
+
 	//~ Begin UEdGraphNode Interface
 	virtual void AllocateDefaultPins() override;
 	virtual FLinearColor GetNodeTitleColor() const override;
@@ -43,7 +47,6 @@ class UK2Node_DynamicCast : public UK2Node
 	virtual void GetNodeContextMenuActions(class UToolMenu* Menu, class UGraphNodeContextMenuContext* Context) const override;
 	virtual bool IncludeParentNodeContextMenu() const override { return true; }
 	virtual void PostReconstructNode() override;
-	virtual void PostPlacedNewNode() override;
 	virtual bool HasExternalDependencies(TArray<class UStruct*>* OptionalOutput) const override;
 	//~ End UEdGraphNode Interface
 
@@ -52,7 +55,7 @@ class UK2Node_DynamicCast : public UK2Node
 	virtual class FNodeHandlingFunctor* CreateNodeHandler(class FKismetCompilerContext& CompilerContext) const override;
 	virtual FText GetMenuCategory() const override;
 	virtual FBlueprintNodeSignature GetSignature() const override;
-	virtual bool IsNodePure() const override { return bIsPureCast; }
+	virtual bool IsNodePure() const override { return PureState == EPureState::Pure; }
 	virtual bool IsConnectionDisallowed(const UEdGraphPin* MyPin, const UEdGraphPin* OtherPin, FString& OutReason) const override;
 	virtual void NotifyPinConnectionListChanged(UEdGraphPin* Pin) override;
 	virtual void ReallocatePinsDuringReconstruction(TArray<UEdGraphPin*>& OldPins) override;
@@ -79,11 +82,17 @@ class UK2Node_DynamicCast : public UK2Node
 	 * Will change the node's purity, and reallocate pins accordingly (adding/
 	 * removing exec pins).
 	 * 
-	 * @param  bNewPurity  The new value for bIsPureCast.
+	 * @param  bNewPurity  Whether to configure node as pure.
 	 */
 	BLUEPRINTGRAPH_API void SetPurity(bool bNewPurity);
 
 protected:
+	/** Determine purity and create exec pins if needed */
+	void CreateExecPins();
+
+	/** Create success output signal pin (visible only if pure) */
+	void CreateSuccessPin();
+
 	/** Flips the node's purity (adding/removing exec pins as needed). */
 	void TogglePurity();
 
@@ -93,7 +102,23 @@ protected:
 	/** Constructing FText strings can be costly, so we cache the node's title */
 	FNodeTextCache CachedNodeTitle;
 
-	UPROPERTY()
+	/** [DEPRECATED] This member is no longer in use. */
+	UE_DEPRECATED(5.3, "Use IsNodePure() instead.")
 	bool bIsPureCast;
+
+private:
+	enum class EPureState : uint8
+	{
+		Pure,
+		Impure,
+		UseDefault
+	};
+
+	/** Internal state used to determine purity */
+	EPureState PureState;
+
+	/** [DEPRECATED] Receives the old value on load if previously saved */
+	UPROPERTY()
+	bool bIsPureCast_DEPRECATED;
 };
 
diff --git a/Engine/Source/Editor/BlueprintGraph/Private/EdGraphSchema_K2.cpp b/Engine/Source/Editor/BlueprintGraph/Private/EdGraphSchema_K2.cpp
--- a/Engine/Source/Editor/BlueprintGraph/Private/EdGraphSchema_K2.cpp
+++ b/Engine/Source/Editor/BlueprintGraph/Private/EdGraphSchema_K2.cpp
@@ -2731,7 +2731,7 @@ TOptional<UEdGraphSchema_K2::FFindSpecializedConversionNodeResults> UEdGraphSche
 	const bool bTryAlternateObjectProperty =
 		InputPin.GetOwningNode()->IsA(UK2Node_CallFunction::StaticClass()) &&
 		IsSelfPin(InputPin) &&
-		((OutputPinType.PinCategory == PC_Object) || (OutputPinType.PinCategory == PC_Interface));
+		((OutputPinType.PinCategory == PC_Object) || ((OutputPinType.PinCategory == PC_Interface) && !OutputPinType.IsContainer()));
 
 	if (bConvertScalarToArray)
 	{
@@ -2809,8 +2809,10 @@ TOptional<UEdGraphSchema_K2::FFindSpecializedConversionNodeResults> UEdGraphSche
 				Result->TargetNode = CastByteToEnum;
 			}
 		}
-		else
+		else if (!OutputPinType.IsContainer())
 		{
+			// Note: Cast nodes do not support a ForEach-style expansion, so we exclude container types here.
+
 			UClass* InputClass  = FBlueprintEditorUtils::GetTypeForPin(InputPin);
 			const UClass* OutputClass = Cast<UClass>(OutputPinType.PinSubCategoryObject.Get());
 
diff --git a/Engine/Source/Editor/BlueprintGraph/Private/K2Node_CallFunction.cpp b/Engine/Source/Editor/BlueprintGraph/Private/K2Node_CallFunction.cpp
--- a/Engine/Source/Editor/BlueprintGraph/Private/K2Node_CallFunction.cpp
+++ b/Engine/Source/Editor/BlueprintGraph/Private/K2Node_CallFunction.cpp
@@ -808,6 +808,66 @@ void UK2Node_CallFunction::ReallocatePinsDuringReconstruction(TArray<UEdGraphPin
 	}
 	// END TEMP
 
+	const UBlueprint* Blueprint = GetBlueprint();
+	if (Blueprint && Blueprint->bIsRegeneratingOnLoad)
+	{
+		// Older nodes incorrectly used an interface context for the target pin for calls to locally-implemented interface
+		// functions (due to an earlier regression). This was compounded by occasional confusion in the context menu, where
+		// a user might have chosen the wrong calling context for an interface implementation and then linked a term with an
+		// incompatible object type to the target input. At runtime this worked fine because both the calling context as well
+		// as the linked object context both implemented the target interface, which allowed any external object context to
+		// be linked to the target pin, so long as it also implemented the interface. After the target pin context was fixed
+		// to match the function context, the target pin context no longer matched the linked pin's context, so the connection
+		// would otherwise be orphaned and result in a Blueprint compiler error. So, rather than cause confusion about nodes
+		// that were not previously broken prior to the fix, we change the context to an interface for backwards compatibility.
+		const UEdGraphPin* OldSelfPin = FindSelfPin(OldPins);
+		if (OldSelfPin
+			&& OldSelfPin->PinType.PinCategory == UEdGraphSchema_K2::PC_Interface
+			&& OldSelfPin->LinkedTo.Num() > 0)
+		{
+			if (const UFunction* TargetFunction = GetTargetFunction())
+			{
+				// Get the function context. This should match the target pin context, but due to the regression that's noted above,
+				// in older assets this may not match the self pin if the context represents an implementation of an interface function.
+				const UClass* FunctionContext = TargetFunction->GetOwnerClass()->GetAuthoritativeClass();
+
+				// Get the interface context from the old target pin (this should already be non-NULL, but we'll check that below).
+				const UClass* InterfaceContext = Cast<UClass>(OldSelfPin->PinType.PinSubCategoryObject);
+
+				// If we're not already using an interface context, but the function's outer class implements the old target pin's type...
+				if (ensure(FunctionContext)
+					&& !FunctionContext->IsChildOf<UInterface>()
+					&& ensure(InterfaceContext)
+					&& FunctionContext->ImplementsInterface(InterfaceContext))
+				{
+					// Check for any linked object pins that aren't compatible with the current function context.
+					for (const UEdGraphPin* LinkedTo : OldSelfPin->LinkedTo)
+					{
+						if (ensure(LinkedTo)
+							&& LinkedTo->PinType.PinCategory == UEdGraphSchema_K2::PC_Object)
+						{
+							// If any linked object pin is not compatible with the current function context, but implements the old target pin's interface,
+							// reset the function context to reference the interface method instead. This will change the function call to an interface call,
+							// which will be compatible with any object context that's being passed in. When we reallocate the pins below, we'll reconstruct
+							// the node using the proper function context in which all linked pins will remain backwards-compatible with the existing target.
+							const UClass* LinkedToPinContext = Cast<UClass>(LinkedTo->PinType.PinSubCategoryObject);
+							if (LinkedToPinContext
+								&& !LinkedToPinContext->IsChildOf(FunctionContext)
+								&& LinkedToPinContext->ImplementsInterface(InterfaceContext))
+							{
+								if (UFunction* InterfaceFunction = FindUField<UFunction>(InterfaceContext, FunctionReference.GetMemberName()))
+								{
+									SetFromFunction(InterfaceFunction);
+									break;
+								}
+							}
+						}
+					}
+				}
+			}
+		}
+	}
+
 	Super::ReallocatePinsDuringReconstruction(OldPins);
 
 	// Connect Execute and Then pins for functions, which became pure.
@@ -2334,6 +2394,24 @@ void UK2Node_CallFunction::ValidateNodeDuringCompilation(class FCompilerResultsL
 			).ToString();
 			MessageLog.Error(*ErrorString, this);
 		}
+		else if (Pin && Pin->PinType.PinCategory == UEdGraphSchema_K2::PC_Object && Pin->PinName == UEdGraphSchema_K2::PN_Self)
+		{
+			const UEdGraphPin* SelfPin = MessageLog.FindSourcePin(Pin);
+			for (const UEdGraphPin* LinkedTo : SelfPin->LinkedTo)
+			{
+				if (ensure(LinkedTo) && LinkedTo->PinType.PinCategory == UEdGraphSchema_K2::PC_Interface)
+				{
+					if (LinkedTo->PinType.IsArray())
+					{
+						MessageLog.Note(*FText::Format(LOCTEXT("InterfaceArrayTargetConnectionNote", "@@: An array of interface types can no longer be directly connected to '{0}'. Each entry must first be cast to the object type. However, the existing connection to '{1}' will continue to work for backwards-compatibility."), FText::FromString(SelfPin->GetName()), FText::FromString(LinkedTo->GetName())).ToString(), this);
+					}
+					else
+					{
+						MessageLog.Note(*FText::Format(LOCTEXT("InterfaceTargetConnectionNote", "@@: An interface type can no longer be directly connected to '{0}'. It must first be cast to the object type. However, the existing connection to '{1}' will continue to work for backwards-compatibility."), FText::FromString(SelfPin->GetName()), FText::FromString(LinkedTo->GetName())).ToString(), this);
+					}
+				}
+			}
+		}
 	}
 }
 
@@ -2833,18 +2911,34 @@ void UK2Node_CallFunction::ExpandNode(class FKismetCompilerContext& CompilerCont
 		}
 	}
 
+	// Older assets may have interface pins wired directly to the target pin in the source graph (due to an earlier regression).
+	UEdGraphPin* SelfPin = Schema->FindSelfPin(*this, EEdGraphPinDirection::EGPD_Input);
+	if (SelfPin && SelfPin->PinType.PinCategory == UEdGraphSchema_K2::PC_Object)
+	{
+		for (UEdGraphPin* PinLinkedToSelfPin : SelfPin->LinkedTo)
+		{
+			if (PinLinkedToSelfPin && PinLinkedToSelfPin->PinType.PinCategory == UEdGraphSchema_K2::PC_Interface && !PinLinkedToSelfPin->PinType.IsContainer())
+			{
+				PinLinkedToSelfPin->BreakLinkTo(SelfPin);
+				if (!Schema->TryCreateConnection(PinLinkedToSelfPin, SelfPin))
+				{
+					PinLinkedToSelfPin->MakeLinkTo(SelfPin);
+				}
+			}
+		}
+	}
+
 	// Then we go through and expand out array iteration if necessary
 	const bool bAllowMultipleSelfs = AllowMultipleSelfs(true);
-	UEdGraphPin* MultiSelf = Schema->FindSelfPin(*this, EEdGraphPinDirection::EGPD_Input);
-	if(bAllowMultipleSelfs && MultiSelf && !MultiSelf->PinType.IsArray())
+	if (bAllowMultipleSelfs && SelfPin && !SelfPin->PinType.IsArray())
 	{
-		const bool bProperInputToExpandForEach = 
-			(1 == MultiSelf->LinkedTo.Num()) && 
-			(nullptr != MultiSelf->LinkedTo[0]) && 
-			(MultiSelf->LinkedTo[0]->PinType.IsArray());
-		if(bProperInputToExpandForEach)
+		const bool bProperInputToExpandForEach =
+			(1 == SelfPin->LinkedTo.Num()) &&
+			(nullptr != SelfPin->LinkedTo[0]) &&
+			(SelfPin->LinkedTo[0]->PinType.IsArray());
+		if (bProperInputToExpandForEach)
 		{
-			CallForEachElementInArrayExpansion(this, MultiSelf, CompilerContext, SourceGraph);
+			CallForEachElementInArrayExpansion(this, SelfPin, CompilerContext, SourceGraph);
 		}
 	}
 }
diff --git a/Engine/Source/Editor/BlueprintGraph/Private/K2Node_ClassDynamicCast.cpp b/Engine/Source/Editor/BlueprintGraph/Private/K2Node_ClassDynamicCast.cpp
--- a/Engine/Source/Editor/BlueprintGraph/Private/K2Node_ClassDynamicCast.cpp
+++ b/Engine/Source/Editor/BlueprintGraph/Private/K2Node_ClassDynamicCast.cpp
@@ -49,22 +49,8 @@ void UK2Node_ClassDynamicCast::AllocateDefaultPins()
 	//@TODO: Move this somewhere more sensible
 	ensure((TargetType == nullptr) || (!TargetType->HasAnyClassFlags(CLASS_NewerVersionExists)));
 
-	const UEdGraphSchema_K2* K2Schema = Cast<UEdGraphSchema_K2>(GetSchema());
-	check(K2Schema != nullptr);
-	if (!K2Schema->DoesGraphSupportImpureFunctions(GetGraph()))
-	{
-		bIsPureCast = true;
-	}
-
-	if (!bIsPureCast)
-	{
-		// Input - Execution Pin
-		CreatePin(EGPD_Input, UEdGraphSchema_K2::PC_Exec, UEdGraphSchema_K2::PN_Execute);
-
-		// Output - Execution Pins
-		CreatePin(EGPD_Output, UEdGraphSchema_K2::PC_Exec, UEdGraphSchema_K2::PN_CastSucceeded);
-		CreatePin(EGPD_Output, UEdGraphSchema_K2::PC_Exec, UEdGraphSchema_K2::PN_CastFailed);
-	}
+	// Exec pins (if needed)
+	CreateExecPins();
 
 	// Input - Source type Pin
 	CreatePin(EGPD_Input, UEdGraphSchema_K2::PC_Class, UObject::StaticClass(), FClassDynamicCastHelper::ClassToCastName);
@@ -76,8 +62,8 @@ void UK2Node_ClassDynamicCast::AllocateDefaultPins()
 		CreatePin(EGPD_Output, UEdGraphSchema_K2::PC_Class, *TargetType, *CastResultPinName);
 	}
 
-	UEdGraphPin* BoolSuccessPin = CreatePin(EGPD_Output, UEdGraphSchema_K2::PC_Boolean, FClassDynamicCastHelper::CastSuccessPinName);
-	BoolSuccessPin->bHidden = !bIsPureCast;
+	// Output - Success
+	CreateSuccessPin();
 
 	UK2Node::AllocateDefaultPins();
 }
diff --git a/Engine/Source/Editor/BlueprintGraph/Private/K2Node_DynamicCast.cpp b/Engine/Source/Editor/BlueprintGraph/Private/K2Node_DynamicCast.cpp
--- a/Engine/Source/Editor/BlueprintGraph/Private/K2Node_DynamicCast.cpp
+++ b/Engine/Source/Editor/BlueprintGraph/Private/K2Node_DynamicCast.cpp
@@ -30,6 +30,7 @@
 #include "ToolMenu.h"
 #include "ToolMenuSection.h"
 #include "UObject/Class.h"
+#include "UObject/FortniteMainBranchObjectVersion.h"
 #include "UObject/Interface.h"
 #include "UObject/NameTypes.h"
 #include "UObject/Object.h"
@@ -48,27 +49,62 @@ namespace UK2Node_DynamicCastImpl
 
 UK2Node_DynamicCast::UK2Node_DynamicCast(const FObjectInitializer& ObjectInitializer)
 	: Super(ObjectInitializer)
-	, bIsPureCast(false)
+	, PureState(EPureState::UseDefault)
 {
+PRAGMA_DISABLE_DEPRECATION_WARNINGS
+	// @todo_deprecated - Remove these later.
+	bIsPureCast = false;
+	bIsPureCast_DEPRECATED = bIsPureCast;
+PRAGMA_ENABLE_DEPRECATION_WARNINGS
 }
 
-void UK2Node_DynamicCast::AllocateDefaultPins()
+void UK2Node_DynamicCast::Serialize(FArchive& Ar)
 {
-	const bool bReferenceObsoleteClass = TargetType && TargetType->HasAnyClassFlags(CLASS_NewerVersionExists);
-	if (bReferenceObsoleteClass)
+	Super::Serialize(Ar);
+
+	Ar.UsingCustomVersion(FFortniteMainBranchObjectVersion::GUID);
+
+	if (Ar.IsLoading() && Ar.CustomVer(FFortniteMainBranchObjectVersion::GUID) < FFortniteMainBranchObjectVersion::DynamicCastNodesUsePureStateEnum)
 	{
-		Message_Error(FString::Printf(TEXT("Node '%s' references obsolete class '%s'"), *GetPathName(), *TargetType->GetPathName()));
+		if (bIsPureCast_DEPRECATED)
+		{
+			PureState = EPureState::Pure;
+		}
+		else
+		{
+			PureState = EPureState::Impure;
+		}
 	}
-	ensure(!bReferenceObsoleteClass);
+	else
+	{
+		Ar << PureState;
+	}
+}
 
+void UK2Node_DynamicCast::CreateExecPins()
+{
+	// Ensure the node is either pure or impure, based on current settings.
 	const UEdGraphSchema_K2* K2Schema = Cast<UEdGraphSchema_K2>(GetSchema());
 	check(K2Schema != nullptr);
 	if (!K2Schema->DoesGraphSupportImpureFunctions(GetGraph()))
 	{
-		bIsPureCast = true;
+		PureState = EPureState::Pure;
+	}
+	else if (PureState == EPureState::UseDefault)
+	{
+		const UBlueprintEditorSettings* BlueprintSettings = GetDefault<UBlueprintEditorSettings>();
+		if (BlueprintSettings->bFavorPureCastNodes)
+		{
+			PureState = EPureState::Pure;
+		}
+		else
+		{
+			PureState = EPureState::Impure;
+		}
 	}
 
-	if (!bIsPureCast)
+	const bool bIsNodePure = IsNodePure();
+	if (!bIsNodePure)
 	{
 		// Input - Execution Pin
 		CreatePin(EGPD_Input, UEdGraphSchema_K2::PC_Exec, UEdGraphSchema_K2::PN_Execute);
@@ -78,6 +114,30 @@ void UK2Node_DynamicCast::AllocateDefaultPins()
 		CreatePin(EGPD_Output, UEdGraphSchema_K2::PC_Exec, UEdGraphSchema_K2::PN_CastFailed);
 	}
 
+PRAGMA_DISABLE_DEPRECATION_WARNINGS
+	// @todo_deprecated - Remove this later.
+	bIsPureCast = bIsNodePure;
+PRAGMA_ENABLE_DEPRECATION_WARNINGS
+}
+
+void UK2Node_DynamicCast::CreateSuccessPin()
+{
+	UEdGraphPin* BoolSuccessPin = CreatePin(EGPD_Output, UEdGraphSchema_K2::PC_Boolean, UK2Node_DynamicCastImpl::CastSuccessPinName);
+	BoolSuccessPin->bHidden = !IsNodePure();
+}
+
+void UK2Node_DynamicCast::AllocateDefaultPins()
+{
+	const bool bReferenceObsoleteClass = TargetType && TargetType->HasAnyClassFlags(CLASS_NewerVersionExists);
+	if (bReferenceObsoleteClass)
+	{
+		Message_Error(FString::Printf(TEXT("Node '%s' references obsolete class '%s'"), *GetPathName(), *TargetType->GetPathName()));
+	}
+	ensure(!bReferenceObsoleteClass);
+
+	// Exec pins (if needed)
+	CreateExecPins();
+
 	// Input - Source type Pin
 	CreatePin(EGPD_Input, UEdGraphSchema_K2::PC_Wildcard, UObject::StaticClass(), UEdGraphSchema_K2::PN_ObjectToCast);
 
@@ -95,8 +155,8 @@ void UK2Node_DynamicCast::AllocateDefaultPins()
 		}
 	}
 
-	UEdGraphPin* BoolSuccessPin = CreatePin(EGPD_Output, UEdGraphSchema_K2::PC_Boolean, UK2Node_DynamicCastImpl::CastSuccessPinName);
-	BoolSuccessPin->bHidden = !bIsPureCast;
+	// Output - Success
+	CreateSuccessPin();
 
 	Super::AllocateDefaultPins();
 }
@@ -175,7 +235,7 @@ void UK2Node_DynamicCast::GetNodeContextMenuActions(UToolMenu* Menu, UGraphNodeC
 				return bInCanTogglePurity;
 			};
 
-			if (bIsPureCast)
+			if (IsNodePure())
 			{
 				MenuEntryTitle = LOCTEXT("MakeImpureTitle", "Convert to impure cast");
 				MenuEntryTooltip = LOCTEXT("MakeImpureTooltip", "Adds in branching execution pins so that you can separatly handle when the cast fails/succeeds.");
@@ -213,18 +273,10 @@ void UK2Node_DynamicCast::PostReconstructNode()
 	NotifyPinConnectionListChanged(GetCastSourcePin());
 }
 
-void UK2Node_DynamicCast::PostPlacedNewNode()
-{
-	Super::PostPlacedNewNode();
-
-	const UBlueprintEditorSettings* BlueprintSettings = GetDefault<UBlueprintEditorSettings>();
-	SetPurity(BlueprintSettings->bFavorPureCastNodes);
-}
-
 UEdGraphPin* UK2Node_DynamicCast::GetValidCastPin() const
 {
 	UEdGraphPin* Pin = FindPin(UEdGraphSchema_K2::PN_CastSucceeded);
-	check((Pin != nullptr) || bIsPureCast);
+	check((Pin != nullptr) || IsNodePure());
 	check((Pin == nullptr) || (Pin->Direction == EGPD_Output));
 	return Pin;
 }
@@ -232,7 +284,7 @@ UEdGraphPin* UK2Node_DynamicCast::GetValidCastPin() const
 UEdGraphPin* UK2Node_DynamicCast::GetInvalidCastPin() const
 {
 	UEdGraphPin* Pin = FindPin(UEdGraphSchema_K2::PN_CastFailed);
-	check((Pin != nullptr) || bIsPureCast);
+	check((Pin != nullptr) || IsNodePure());
 	check((Pin == nullptr) || (Pin->Direction == EGPD_Output));
 	return Pin;
 }
@@ -272,11 +324,23 @@ UEdGraphPin* UK2Node_DynamicCast::GetBoolSuccessPin() const
 
 void UK2Node_DynamicCast::SetPurity(bool bNewPurity)
 {
-	if (bNewPurity != bIsPureCast)
+	if (bNewPurity != IsNodePure())
 	{
+		if (bNewPurity)
+		{
+			PureState = EPureState::Pure;
+		}
+		else
+		{
+			PureState = EPureState::Impure;
+		}
+
+PRAGMA_DISABLE_DEPRECATION_WARNINGS
+		// @todo_deprecated - Remove this later.
 		bIsPureCast = bNewPurity;
+PRAGMA_ENABLE_DEPRECATION_WARNINGS
 
-		bool const bHasBeenConstructed = (Pins.Num() > 0);
+		const bool bHasBeenConstructed = (Pins.Num() > 0);
 		if (bHasBeenConstructed)
 		{
 			ReconstructNode();
@@ -286,11 +350,12 @@ void UK2Node_DynamicCast::SetPurity(bool bNewPurity)
 
 void UK2Node_DynamicCast::TogglePurity()
 {
-	const FText TransactionTitle = bIsPureCast ? LOCTEXT("TogglePurityToImpure", "Convert to Impure Cast") : LOCTEXT("TogglePurityToPure", "Convert to Pure Cast");
+	const bool bIsNodePure = IsNodePure();
+	const FText TransactionTitle = bIsNodePure ? LOCTEXT("TogglePurityToImpure", "Convert to Impure Cast") : LOCTEXT("TogglePurityToPure", "Convert to Pure Cast");
 	const FScopedTransaction Transaction( TransactionTitle );
 	Modify();
 
-	SetPurity(!bIsPureCast);
+	SetPurity(!bIsNodePure);
 }
 
 UK2Node::ERedirectType UK2Node_DynamicCast::DoPinsMatchForReconstruction(const UEdGraphPin* NewPin, int32 NewPinIndex, const UEdGraphPin* OldPin, int32 OldPinIndex) const
@@ -496,7 +561,7 @@ void UK2Node_DynamicCast::ValidateNodeDuringCompilation(FCompilerResultsLog& Mes
 
 bool UK2Node_DynamicCast::ReconnectPureExecPins(TArray<UEdGraphPin*>& OldPins)
 {
-	if (bIsPureCast)
+	if (IsNodePure())
 	{
 		// look for an old exec pin
 		UEdGraphPin* PinExec = nullptr;
diff --git a/Engine/Source/Editor/Kismet/Private/BlueprintCompilationManager.cpp b/Engine/Source/Editor/Kismet/Private/BlueprintCompilationManager.cpp
--- a/Engine/Source/Editor/Kismet/Private/BlueprintCompilationManager.cpp
+++ b/Engine/Source/Editor/Kismet/Private/BlueprintCompilationManager.cpp
@@ -2731,6 +2731,11 @@ UClass* FBlueprintCompilationManagerImpl::FastGenerateSkeletonClass(UBlueprint*
 		}
 
 		UFunction* ParentFn = ParentClass->FindFunctionByName(NewFunction->GetFName());
+
+		// Set the parent function prior to searching for a corresponding interface class function. This matches what
+		// is done in the full path (@see FKismetCompilerContext::PrecompileFunction). It is ok for this to be NULL.
+		NewFunction->SetSuperStruct(ParentFn);
+
 		if(ParentFn == nullptr)
 		{
 			// check for function in implemented interfaces:
@@ -2758,7 +2763,6 @@ UClass* FBlueprintCompilationManagerImpl::FastGenerateSkeletonClass(UBlueprint*
 				}
 			}
 		}
-		NewFunction->SetSuperStruct( ParentFn );
 		
 		InCurrentParamStorageLocation = &NewFunction->ChildProperties;
 
diff --git a/Engine/Source/Runtime/Core/Public/UObject/FortniteMainBranchObjectVersion.h b/Engine/Source/Runtime/Core/Public/UObject/FortniteMainBranchObjectVersion.h
--- a/Engine/Source/Runtime/Core/Public/UObject/FortniteMainBranchObjectVersion.h
+++ b/Engine/Source/Runtime/Core/Public/UObject/FortniteMainBranchObjectVersion.h
@@ -268,6 +268,9 @@ struct CORE_API FFortniteMainBranchObjectVersion
 		// Add an attribute to geometry collection to track internal faces, rather than relying on material ID numbering
 		ChaosGeometryCollectionInternalFacesAttribute,
 
+		// Dynamic cast nodes use an enumerated pure node state to include a value for the default setting
+		DynamicCastNodesUsePureStateEnum,
+
 		// -----<new versions can be added above this line>-------------------------------------------------
 		VersionPlusOne,
 		LatestVersion = VersionPlusOne - 1
