Fix conversions when using the "self" node.

When a pin type is PSC_Self, then its object class type is assumed to be that of the generated class. FindSpecializedConversionNode performed that check for the input pin, but *not* the output pin. This would prevent linking of the self node's pin to any input pin that would need a conversion (eg: linking the self node's pin to a soft reference pin).

If the output pin type is a self type, then the fix here simply changes the output pin type to use the generated class obtained via the input pin. Note that this assumes that the output pin's type also comes from a pin in the same Blueprint as the input.

diff --git a/Engine/Source/Editor/BlueprintGraph/Private/EdGraphSchema_K2.cpp b/Engine/Source/Editor/BlueprintGraph/Private/EdGraphSchema_K2.cpp
--- a/Engine/Source/Editor/BlueprintGraph/Private/EdGraphSchema_K2.cpp
+++ b/Engine/Source/Editor/BlueprintGraph/Private/EdGraphSchema_K2.cpp
@@ -2812,10 +2812,24 @@ TOptional<UEdGraphSchema_K2::FFindSpecializedConversionNodeResults> UEdGraphSche
 		else if (!OutputPinType.IsContainer())
 		{
 			// Note: Cast nodes do not support a ForEach-style expansion, so we exclude container types here.
+		
+			const UBlueprint* Blueprint = FBlueprintEditorUtils::FindBlueprintForNodeChecked(InputPin.GetOwningNode());
+			UClass* BlueprintClass = (Blueprint->GeneratedClass != nullptr) ? Blueprint->GeneratedClass : Blueprint->ParentClass;
+
+			UClass* InputClass = Cast<UClass>(InputPin.PinType.PinSubCategoryObject.Get());
+
+			if ((InputClass == nullptr) && (InputPin.PinType.PinSubCategory == UEdGraphSchema_K2::PSC_Self))
+			{
+				InputClass = BlueprintClass;
+			}
 
-			UClass* InputClass  = FBlueprintEditorUtils::GetTypeForPin(InputPin);
 			const UClass* OutputClass = Cast<UClass>(OutputPinType.PinSubCategoryObject.Get());
 
+			if ((OutputClass == nullptr) && (OutputPinType.PinSubCategory == UEdGraphSchema_K2::PSC_Self))
+			{
+				OutputClass = BlueprintClass;
+			}
+
 			bool bNeedsDynamicCast = false;
 			if ((OutputPinType.PinCategory == PC_Interface) && (InputPinType.PinCategory == PC_Object))
 			{
