Don't orphan pins during Blueprint graph node fixup if they are linked to an autocast or other implicit conversion node that supercedes the original link.

diff --git a/Engine/Source/Editor/BlueprintGraph/Private/EdGraphSchema_K2.cpp b/Engine/Source/Editor/BlueprintGraph/Private/EdGraphSchema_K2.cpp
--- a/Engine/Source/Editor/BlueprintGraph/Private/EdGraphSchema_K2.cpp
+++ b/Engine/Source/Editor/BlueprintGraph/Private/EdGraphSchema_K2.cpp
@@ -6141,7 +6141,7 @@ struct FBackwardCompatibilityConversionHelper
 
 		for (UEdGraphPin* OldLink : OldLinks)
 		{
-			if (!NewPin || !NewPin->LinkedTo.Contains(OldLink))
+			if (!NewPin || (!NewPin->LinkedTo.Contains(OldLink) && NewNode->GetSchema()->CanCreateConnection(NewPin, OldLink).Response != CONNECT_RESPONSE_MAKE_WITH_CONVERSION_NODE))
 			{
 				OrphanedLinks.Add(OldLink);
 			}
