Lay groundwork for avoiding SGraphPanel refreshes after making a change to a single node by adding a NotifyNodeChanged routine to UEdGraph, currently implemented to just invoke NotifyGraphChanged with no additional context

diff --git a/Engine/Plugins/Editor/GameplayTagsEditor/Source/GameplayTagsEditor/Private/GameplayTagsK2Node_MultiCompareBase.cpp b/Engine/Plugins/Editor/GameplayTagsEditor/Source/GameplayTagsEditor/Private/GameplayTagsK2Node_MultiCompareBase.cpp
--- a/Engine/Plugins/Editor/GameplayTagsEditor/Source/GameplayTagsEditor/Private/GameplayTagsK2Node_MultiCompareBase.cpp
+++ b/Engine/Plugins/Editor/GameplayTagsEditor/Source/GameplayTagsEditor/Private/GameplayTagsK2Node_MultiCompareBase.cpp
@@ -24,7 +24,7 @@ void UGameplayTagsK2Node_MultiCompareBase::PostEditChangeProperty(struct FProper
 		}
 
 		ReconstructNode();
-		GetGraph()->NotifyGraphChanged();
+		GetGraph()->NotifyNodeChanged(this);
 	}
 	Super::PostEditChangeProperty(PropertyChangedEvent);
 }
diff --git a/Engine/Plugins/Editor/GameplayTagsEditor/Source/GameplayTagsEditor/Private/GameplayTagsK2Node_SwitchGameplayTag.cpp b/Engine/Plugins/Editor/GameplayTagsEditor/Source/GameplayTagsEditor/Private/GameplayTagsK2Node_SwitchGameplayTag.cpp
--- a/Engine/Plugins/Editor/GameplayTagsEditor/Source/GameplayTagsEditor/Private/GameplayTagsK2Node_SwitchGameplayTag.cpp
+++ b/Engine/Plugins/Editor/GameplayTagsEditor/Source/GameplayTagsEditor/Private/GameplayTagsK2Node_SwitchGameplayTag.cpp
@@ -63,7 +63,7 @@ void UGameplayTagsK2Node_SwitchGameplayTag::PostEditChangeProperty(struct FPrope
 	if (bIsDirty)
 	{
 		ReconstructNode();
-		GetGraph()->NotifyGraphChanged();
+		GetGraph()->NotifyNodeChanged(this);
 	}
 	Super::PostEditChangeProperty(PropertyChangedEvent);
 }
diff --git a/Engine/Plugins/Editor/GameplayTagsEditor/Source/GameplayTagsEditor/Private/GameplayTagsK2Node_SwitchGameplayTagContainer.cpp b/Engine/Plugins/Editor/GameplayTagsEditor/Source/GameplayTagsEditor/Private/GameplayTagsK2Node_SwitchGameplayTagContainer.cpp
--- a/Engine/Plugins/Editor/GameplayTagsEditor/Source/GameplayTagsEditor/Private/GameplayTagsK2Node_SwitchGameplayTagContainer.cpp
+++ b/Engine/Plugins/Editor/GameplayTagsEditor/Source/GameplayTagsEditor/Private/GameplayTagsK2Node_SwitchGameplayTagContainer.cpp
@@ -64,7 +64,7 @@ void UGameplayTagsK2Node_SwitchGameplayTagContainer::PostEditChangeProperty(stru
 	if (bIsDirty)
 	{
 		ReconstructNode();
-		GetGraph()->NotifyGraphChanged();
+		GetGraph()->NotifyNodeChanged(this);
 	}
 	Super::PostEditChangeProperty(PropertyChangedEvent);
 }
diff --git a/Engine/Plugins/Experimental/PythonScriptPlugin/Source/PythonScriptPlugin/Private/K2Node_ExecutePythonScript.cpp b/Engine/Plugins/Experimental/PythonScriptPlugin/Source/PythonScriptPlugin/Private/K2Node_ExecutePythonScript.cpp
--- a/Engine/Plugins/Experimental/PythonScriptPlugin/Source/PythonScriptPlugin/Private/K2Node_ExecutePythonScript.cpp
+++ b/Engine/Plugins/Experimental/PythonScriptPlugin/Source/PythonScriptPlugin/Private/K2Node_ExecutePythonScript.cpp
@@ -130,7 +130,7 @@ void UK2Node_ExecutePythonScript::SynchronizeArgumentPinTypeImpl(UEdGraphPin* Pi
 	{
 		Pin->PinType = NewPinType;
 
-		GetGraph()->NotifyGraphChanged();
+		GetGraph()->NotifyNodeChanged(this);
 
 		UBlueprint* Blueprint = GetBlueprint();
 		if (!Blueprint->bBeingCompiled)
@@ -167,7 +167,7 @@ void UK2Node_ExecutePythonScript::PostEditChangeProperty(struct FPropertyChanged
 	if (PropertyName == GET_MEMBER_NAME_CHECKED(UK2Node_ExecutePythonScript, Inputs) || PropertyName == GET_MEMBER_NAME_CHECKED(UK2Node_ExecutePythonScript, Outputs))
 	{
 		ReconstructNode();
-		GetGraph()->NotifyGraphChanged();
+		GetGraph()->NotifyNodeChanged(this);
 	}
 	Super::PostEditChangeProperty(PropertyChangedEvent);
 }
diff --git a/Engine/Plugins/JsonBlueprintUtilities/Source/JsonBlueprintGraph/Private/K2Node_GetJsonField.cpp b/Engine/Plugins/JsonBlueprintUtilities/Source/JsonBlueprintGraph/Private/K2Node_GetJsonField.cpp
--- a/Engine/Plugins/JsonBlueprintUtilities/Source/JsonBlueprintGraph/Private/K2Node_GetJsonField.cpp
+++ b/Engine/Plugins/JsonBlueprintUtilities/Source/JsonBlueprintGraph/Private/K2Node_GetJsonField.cpp
@@ -97,7 +97,7 @@ void UK2Node_GetJsonField::NotifyInputChanged() const
 	}
 
 	UEdGraph* Graph = GetGraph();
-	Graph->NotifyGraphChanged();
+	Graph->NotifyNodeChanged(this);
 }
 
 #undef LOCTEXT_NAMESPACE
diff --git a/Engine/Plugins/JsonBlueprintUtilities/Source/JsonBlueprintGraph/Private/K2Node_SetJsonField.cpp b/Engine/Plugins/JsonBlueprintUtilities/Source/JsonBlueprintGraph/Private/K2Node_SetJsonField.cpp
--- a/Engine/Plugins/JsonBlueprintUtilities/Source/JsonBlueprintGraph/Private/K2Node_SetJsonField.cpp
+++ b/Engine/Plugins/JsonBlueprintUtilities/Source/JsonBlueprintGraph/Private/K2Node_SetJsonField.cpp
@@ -96,7 +96,7 @@ void UK2Node_SetJsonField::NotifyInputChanged() const
 	}
 
 	UEdGraph* Graph = GetGraph();
-	Graph->NotifyGraphChanged();
+	Graph->NotifyNodeChanged(this);
 }
 
 #undef LOCTEXT_NAMESPACE
diff --git a/Engine/Plugins/JsonBlueprintUtilities/Source/JsonBlueprintGraph/Private/K2Node_StructToJsonString.cpp b/Engine/Plugins/JsonBlueprintUtilities/Source/JsonBlueprintGraph/Private/K2Node_StructToJsonString.cpp
--- a/Engine/Plugins/JsonBlueprintUtilities/Source/JsonBlueprintGraph/Private/K2Node_StructToJsonString.cpp
+++ b/Engine/Plugins/JsonBlueprintUtilities/Source/JsonBlueprintGraph/Private/K2Node_StructToJsonString.cpp
@@ -1,4 +1,4 @@
-﻿// Copyright Epic Games, Inc. All Rights Reserved.
+// Copyright Epic Games, Inc. All Rights Reserved.
 
 #include "K2Node_StructToJsonString.h"
 
@@ -86,7 +86,7 @@ void UK2Node_StructToJsonString::NotifyInputChanged() const
 	}
 
 	UEdGraph* Graph = GetGraph();
-	Graph->NotifyGraphChanged();
+	Graph->NotifyNodeChanged(this);
 }
 
 #undef LOCTEXT_NAMESPACE
diff --git a/Engine/Plugins/VirtualProduction/DMX/DMXEngine/Source/DMXBlueprintGraph/Private/K2Node_GetDMXFixturePatch.cpp b/Engine/Plugins/VirtualProduction/DMX/DMXEngine/Source/DMXBlueprintGraph/Private/K2Node_GetDMXFixturePatch.cpp
--- a/Engine/Plugins/VirtualProduction/DMX/DMXEngine/Source/DMXBlueprintGraph/Private/K2Node_GetDMXFixturePatch.cpp
+++ b/Engine/Plugins/VirtualProduction/DMX/DMXEngine/Source/DMXBlueprintGraph/Private/K2Node_GetDMXFixturePatch.cpp
@@ -258,7 +258,7 @@ void UK2Node_GetDMXFixturePatch::NotifyInputChanged()
 	}
 
 	UEdGraph* Graph = GetGraph();
-	Graph->NotifyGraphChanged();
+	Graph->NotifyNodeChanged(this);
 }
 
 #undef LOCTEXT_NAMESPACE
diff --git a/Engine/Plugins/VirtualProduction/DMX/DMXEngine/Source/DMXBlueprintGraph/Private/K2Node_GetDMXFixtureType.cpp b/Engine/Plugins/VirtualProduction/DMX/DMXEngine/Source/DMXBlueprintGraph/Private/K2Node_GetDMXFixtureType.cpp
--- a/Engine/Plugins/VirtualProduction/DMX/DMXEngine/Source/DMXBlueprintGraph/Private/K2Node_GetDMXFixtureType.cpp
+++ b/Engine/Plugins/VirtualProduction/DMX/DMXEngine/Source/DMXBlueprintGraph/Private/K2Node_GetDMXFixtureType.cpp
@@ -82,7 +82,7 @@ void UK2Node_GetDMXFixtureType::PinDefaultValueChanged(UEdGraphPin* FromPin)
 	Super::PinDefaultValueChanged(FromPin);
 
 	UEdGraph* Graph = GetGraph();
-	Graph->NotifyGraphChanged();
+	Graph->NotifyNodeChanged(this);
 }
 
 void UK2Node_GetDMXFixtureType::ExpandNode(FKismetCompilerContext& CompilerContext, UEdGraph* SourceGraph)
diff --git a/Engine/Plugins/VirtualProduction/DMX/DMXPixelMapping/Source/DMXPixelMappingBlueprintGraph/Private/K2Node_PixelMappingBaseComponent.cpp b/Engine/Plugins/VirtualProduction/DMX/DMXPixelMapping/Source/DMXPixelMappingBlueprintGraph/Private/K2Node_PixelMappingBaseComponent.cpp
--- a/Engine/Plugins/VirtualProduction/DMX/DMXPixelMapping/Source/DMXPixelMappingBlueprintGraph/Private/K2Node_PixelMappingBaseComponent.cpp
+++ b/Engine/Plugins/VirtualProduction/DMX/DMXPixelMapping/Source/DMXPixelMappingBlueprintGraph/Private/K2Node_PixelMappingBaseComponent.cpp
@@ -29,7 +29,7 @@ void UK2Node_PixelMappingBaseComponent::AllocateDefaultPins()
 void UK2Node_PixelMappingBaseComponent::RefreshGraph()
 {
 	UEdGraph* Graph = GetGraph();
-	Graph->NotifyGraphChanged();
+	Graph->NotifyNodeChanged(this);
 }
 
 UEdGraphPin* UK2Node_PixelMappingBaseComponent::GetInPixelMappingPin() const
diff --git a/Engine/Plugins/Web/HttpBlueprint/Source/HttpBlueprintGraph/Private/K2Node_HttpRequest.cpp b/Engine/Plugins/Web/HttpBlueprint/Source/HttpBlueprintGraph/Private/K2Node_HttpRequest.cpp
--- a/Engine/Plugins/Web/HttpBlueprint/Source/HttpBlueprintGraph/Private/K2Node_HttpRequest.cpp
+++ b/Engine/Plugins/Web/HttpBlueprint/Source/HttpBlueprintGraph/Private/K2Node_HttpRequest.cpp
@@ -406,7 +406,7 @@ void UK2Node_HttpRequest::SyncBodyPinType(UEdGraphPin* const Pin) const
 
 	if (bBodyTypeChanged)
 	{
-		GetGraph()->NotifyGraphChanged();
+		GetGraph()->NotifyNodeChanged(this);
 		if (UBlueprint* const BP = GetBlueprint(); !BP->bBeingCompiled)
 		{
 			FBlueprintEditorUtils::MarkBlueprintAsModified(BP);
diff --git a/Engine/Plugins/Web/HttpBlueprint/Source/HttpBlueprintGraph/Private/K2Node_MakeRequestHeader.cpp b/Engine/Plugins/Web/HttpBlueprint/Source/HttpBlueprintGraph/Private/K2Node_MakeRequestHeader.cpp
--- a/Engine/Plugins/Web/HttpBlueprint/Source/HttpBlueprintGraph/Private/K2Node_MakeRequestHeader.cpp
+++ b/Engine/Plugins/Web/HttpBlueprint/Source/HttpBlueprintGraph/Private/K2Node_MakeRequestHeader.cpp
@@ -418,7 +418,7 @@ void UK2Node_MakeRequestHeader::AddInputPin()
 	if (!GetBlueprint()->bBeingCompiled)
 	{
 		FBlueprintEditorUtils::MarkBlueprintAsStructurallyModified(GetBlueprint());
-		GetGraph()->NotifyGraphChanged();
+		GetGraph()->NotifyNodeChanged(this);
 	}
 }
 
diff --git a/Engine/Source/Editor/BlueprintGraph/Private/K2Node.cpp b/Engine/Source/Editor/BlueprintGraph/Private/K2Node.cpp
--- a/Engine/Source/Editor/BlueprintGraph/Private/K2Node.cpp
+++ b/Engine/Source/Editor/BlueprintGraph/Private/K2Node.cpp
@@ -529,7 +529,7 @@ void UK2Node::PinConnectionListChanged(UEdGraphPin* Pin)
 
 		if (OuterGraph)
 		{
-			OuterGraph->NotifyGraphChanged();
+			OuterGraph->NotifyNodeChanged(this);
 		}
 
 		if (Pin->ParentPin == nullptr)
@@ -780,7 +780,7 @@ void UK2Node::ReconstructNode()
 		}
 	}
 
-	GetGraph()->NotifyGraphChanged();
+	GetGraph()->NotifyNodeChanged(this);
 }
 
 void UK2Node::GetRedirectPinNames(const UEdGraphPin& Pin, TArray<FString>& RedirectPinNames) const
diff --git a/Engine/Source/Editor/BlueprintGraph/Private/K2Node_EnumEquality.cpp b/Engine/Source/Editor/BlueprintGraph/Private/K2Node_EnumEquality.cpp
--- a/Engine/Source/Editor/BlueprintGraph/Private/K2Node_EnumEquality.cpp
+++ b/Engine/Source/Editor/BlueprintGraph/Private/K2Node_EnumEquality.cpp
@@ -91,7 +91,7 @@ void UK2Node_EnumEquality::NotifyPinConnectionListChanged(UEdGraphPin* Pin)
 			Schema->SetPinAutogeneratedDefaultValueBasedOnType(Input1Pin);
 			Schema->SetPinAutogeneratedDefaultValueBasedOnType(Input2Pin);
 			// We have to refresh the graph to get the enum dropdowns to go away
-			GetGraph()->NotifyGraphChanged();
+			GetGraph()->NotifyNodeChanged(this);
 		}
 		else if (Pin->LinkedTo.Num() > 0)
 		{
@@ -122,7 +122,7 @@ void UK2Node_EnumEquality::NotifyPinConnectionListChanged(UEdGraphPin* Pin)
 				if (OtherPin->LinkedTo.Num() <= 0 && OtherPin->DefaultValue.IsEmpty())
 				{
 					Schema->SetPinAutogeneratedDefaultValueBasedOnType(OtherPin);
-					GetGraph()->NotifyGraphChanged();
+					GetGraph()->NotifyNodeChanged(this);
 				}
 			}
 			// A valid enum wasn't used to break the links
diff --git a/Engine/Source/Editor/BlueprintGraph/Private/K2Node_FormatText.cpp b/Engine/Source/Editor/BlueprintGraph/Private/K2Node_FormatText.cpp
--- a/Engine/Source/Editor/BlueprintGraph/Private/K2Node_FormatText.cpp
+++ b/Engine/Source/Editor/BlueprintGraph/Private/K2Node_FormatText.cpp
@@ -109,7 +109,7 @@ void UK2Node_FormatText::SynchronizeArgumentPinType(UEdGraphPin* Pin)
 		if (bPinTypeChanged)
 		{
 			// Let the graph know to refresh
-			GetGraph()->NotifyGraphChanged();
+			GetGraph()->NotifyNodeChanged(this);
 
 			UBlueprint* Blueprint = GetBlueprint();
 			if (!Blueprint->bBeingCompiled)
@@ -152,7 +152,7 @@ void UK2Node_FormatText::PostEditChangeProperty(struct FPropertyChangedEvent& Pr
 	if (PropertyName == GET_MEMBER_NAME_CHECKED(UK2Node_FormatText, PinNames))
 	{
 		ReconstructNode();
-		GetGraph()->NotifyGraphChanged();
+		GetGraph()->NotifyNodeChanged(this);
 	}
 	Super::PostEditChangeProperty(PropertyChangedEvent);
 }
@@ -226,7 +226,7 @@ void UK2Node_FormatText::PinDefaultValueChanged(UEdGraphPin* Pin)
 			}
 		}
 
-		GetGraph()->NotifyGraphChanged();
+		GetGraph()->NotifyNodeChanged(this);
 	}
 }
 
@@ -308,7 +308,7 @@ void UK2Node_FormatText::PostReconstructNode()
 
 			if (NumPinsFixedUp > 0)
 			{
-				GetGraph()->NotifyGraphChanged();
+				GetGraph()->NotifyNodeChanged(this);
 				FBlueprintEditorUtils::MarkBlueprintAsStructurallyModified(GetBlueprint());
 			}
 		}
@@ -618,7 +618,7 @@ void UK2Node_FormatText::AddArgumentPin()
 	PinNames.Add(PinName);
 
 	FBlueprintEditorUtils::MarkBlueprintAsStructurallyModified(GetBlueprint());
-	GetGraph()->NotifyGraphChanged();
+	GetGraph()->NotifyNodeChanged(this);
 }
 
 void UK2Node_FormatText::RemoveArgument(int32 InIndex)
@@ -634,7 +634,7 @@ void UK2Node_FormatText::RemoveArgument(int32 InIndex)
 	PinNames.RemoveAt(InIndex);
 
 	FBlueprintEditorUtils::MarkBlueprintAsStructurallyModified(GetBlueprint());
-	GetGraph()->NotifyGraphChanged();
+	GetGraph()->NotifyNodeChanged(this);
 }
 
 void UK2Node_FormatText::SetArgumentName(int32 InIndex, FName InName)
@@ -653,7 +653,7 @@ void UK2Node_FormatText::SwapArguments(int32 InIndexA, int32 InIndexB)
 	PinNames.Swap(InIndexA, InIndexB);
 
 	ReconstructNode();
-	GetGraph()->NotifyGraphChanged();
+	GetGraph()->NotifyNodeChanged(this);
 
 	FBlueprintEditorUtils::MarkBlueprintAsModified(GetBlueprint());
 }
diff --git a/Engine/Source/Editor/BlueprintGraph/Private/K2Node_GetArrayItem.cpp b/Engine/Source/Editor/BlueprintGraph/Private/K2Node_GetArrayItem.cpp
--- a/Engine/Source/Editor/BlueprintGraph/Private/K2Node_GetArrayItem.cpp
+++ b/Engine/Source/Editor/BlueprintGraph/Private/K2Node_GetArrayItem.cpp
@@ -302,7 +302,7 @@ void UK2Node_GetArrayItem::NotifyPinConnectionListChanged(UEdGraphPin* Pin)
 				if (NewType.PinCategory != UEdGraphSchema_K2::PC_Wildcard)
 				{
 					PropagatePinType(NewType);
-					GetGraph()->NotifyGraphChanged();
+					GetGraph()->NotifyNodeChanged(this);
 				}
 			}
 		}
diff --git a/Engine/Source/Editor/BlueprintGraph/Private/K2Node_GetClassDefaults.cpp b/Engine/Source/Editor/BlueprintGraph/Private/K2Node_GetClassDefaults.cpp
--- a/Engine/Source/Editor/BlueprintGraph/Private/K2Node_GetClassDefaults.cpp
+++ b/Engine/Source/Editor/BlueprintGraph/Private/K2Node_GetClassDefaults.cpp
@@ -619,7 +619,7 @@ void UK2Node_GetClassDefaults::OnClassPinChanged()
 	// Notify the graph that the node has been changed
 	if(UEdGraph* Graph = GetGraph())
 	{
-		Graph->NotifyGraphChanged();
+		Graph->NotifyNodeChanged(this);
 	}
 }
 
diff --git a/Engine/Source/Editor/BlueprintGraph/Private/K2Node_GetDataTableRow.cpp b/Engine/Source/Editor/BlueprintGraph/Private/K2Node_GetDataTableRow.cpp
--- a/Engine/Source/Editor/BlueprintGraph/Private/K2Node_GetDataTableRow.cpp
+++ b/Engine/Source/Editor/BlueprintGraph/Private/K2Node_GetDataTableRow.cpp
@@ -104,7 +104,7 @@ void UK2Node_GetDataTableRow::RefreshRowNameOptions()
 {
 	// When the DataTable pin gets a new value assigned, we need to update the Slate UI so that SGraphNodeCallParameterCollectionFunction will update the ParameterName drop down
 	UEdGraph* Graph = GetGraph();
-	Graph->NotifyGraphChanged();
+	Graph->NotifyNodeChanged(this);
 }
 
 
diff --git a/Engine/Source/Editor/BlueprintGraph/Private/K2Node_MakeContainer.cpp b/Engine/Source/Editor/BlueprintGraph/Private/K2Node_MakeContainer.cpp
--- a/Engine/Source/Editor/BlueprintGraph/Private/K2Node_MakeContainer.cpp
+++ b/Engine/Source/Editor/BlueprintGraph/Private/K2Node_MakeContainer.cpp
@@ -304,7 +304,7 @@ void UK2Node_MakeContainer::NotifyPinConnectionListChanged(UEdGraphPin* Pin)
 	if (bNotifyGraphChanged)
 	{
 		PropagatePinType();
-		GetGraph()->NotifyGraphChanged();
+		GetGraph()->NotifyNodeChanged(this);
 	}
 }
 
@@ -388,7 +388,7 @@ void UK2Node_MakeContainer::PropagatePinType()
 		{
 			if (UEdGraph* Graph = OwningNode->GetGraph())
 			{
-				Graph->NotifyGraphChanged();
+				Graph->NotifyNodeChanged(this);
 			}
 		}
 	}
diff --git a/Engine/Source/Editor/BlueprintGraph/Private/K2Node_PromotableOperator.cpp b/Engine/Source/Editor/BlueprintGraph/Private/K2Node_PromotableOperator.cpp
--- a/Engine/Source/Editor/BlueprintGraph/Private/K2Node_PromotableOperator.cpp
+++ b/Engine/Source/Editor/BlueprintGraph/Private/K2Node_PromotableOperator.cpp
@@ -310,7 +310,7 @@ void UK2Node_PromotableOperator::NodeConnectionListChanged()
 		UpdateFromBestMatchingFunction();
 
 		// Get correct default value boxes
-		GetGraph()->NotifyGraphChanged();
+		GetGraph()->NotifyNodeChanged(this);
 	}
 }
 
@@ -1116,7 +1116,7 @@ void UK2Node_PromotableOperator::ResetNodeToWildcard()
 		}
 	}
 
-	GetGraph()->NotifyGraphChanged();
+	GetGraph()->NotifyNodeChanged(this);
 }
 
 void UK2Node_PromotableOperator::RecombineAllSplitPins()
@@ -1375,7 +1375,7 @@ void UK2Node_PromotableOperator::UpdatePinsFromFunction(const UFunction* Functio
 
 	// We need to notify the graph that the node has changed to get 
 	// the correct default value text boxes on the node
-	GetGraph()->NotifyGraphChanged();
+	GetGraph()->NotifyNodeChanged(this);
 }
 
 UEdGraphPin* UK2Node_PromotableOperator::GetOutputPin() const
diff --git a/Engine/Source/Editor/BlueprintGraph/Private/K2Node_SpawnActor.cpp b/Engine/Source/Editor/BlueprintGraph/Private/K2Node_SpawnActor.cpp
--- a/Engine/Source/Editor/BlueprintGraph/Private/K2Node_SpawnActor.cpp
+++ b/Engine/Source/Editor/BlueprintGraph/Private/K2Node_SpawnActor.cpp
@@ -202,7 +202,7 @@ void UK2Node_SpawnActor::PinDefaultValueChanged(UEdGraphPin* ChangedPin)
 
 		// Refresh the UI for the graph so the pin changes show up
 		UEdGraph* Graph = GetGraph();
-		Graph->NotifyGraphChanged();
+		Graph->NotifyNodeChanged(this);
 
 		// Mark dirty
 		FBlueprintEditorUtils::MarkBlueprintAsModified(GetBlueprint());
diff --git a/Engine/Source/Editor/BlueprintGraph/Private/K2Node_SwitchName.cpp b/Engine/Source/Editor/BlueprintGraph/Private/K2Node_SwitchName.cpp
--- a/Engine/Source/Editor/BlueprintGraph/Private/K2Node_SwitchName.cpp
+++ b/Engine/Source/Editor/BlueprintGraph/Private/K2Node_SwitchName.cpp
@@ -35,7 +35,7 @@ void UK2Node_SwitchName::PostEditChangeProperty(struct FPropertyChangedEvent& Pr
 	if (bIsDirty)
 	{
 		ReconstructNode();
-		GetGraph()->NotifyGraphChanged();
+		GetGraph()->NotifyNodeChanged(this);
 	}
 	Super::PostEditChangeProperty(PropertyChangedEvent);
 }
diff --git a/Engine/Source/Editor/BlueprintGraph/Private/K2Node_SwitchString.cpp b/Engine/Source/Editor/BlueprintGraph/Private/K2Node_SwitchString.cpp
--- a/Engine/Source/Editor/BlueprintGraph/Private/K2Node_SwitchString.cpp
+++ b/Engine/Source/Editor/BlueprintGraph/Private/K2Node_SwitchString.cpp
@@ -47,7 +47,7 @@ void UK2Node_SwitchString::PostEditChangeProperty(struct FPropertyChangedEvent&
 	if (bIsDirty)
 	{
 		ReconstructNode();
-		GetGraph()->NotifyGraphChanged();
+		GetGraph()->NotifyNodeChanged(this);
 	}
 	Super::PostEditChangeProperty(PropertyChangedEvent);
 }
diff --git a/Engine/Source/Editor/BlueprintGraph/Private/K2Node_VariableSetRef.cpp b/Engine/Source/Editor/BlueprintGraph/Private/K2Node_VariableSetRef.cpp
--- a/Engine/Source/Editor/BlueprintGraph/Private/K2Node_VariableSetRef.cpp
+++ b/Engine/Source/Editor/BlueprintGraph/Private/K2Node_VariableSetRef.cpp
@@ -282,7 +282,7 @@ void UK2Node_VariableSetRef::NotifyPinConnectionListChanged(UEdGraphPin* Pin)
 		CachedNodeTitle.MarkDirty();
 		
 		// Get the graph to refresh our title and default value info
-		GetGraph()->NotifyGraphChanged();
+		GetGraph()->NotifyNodeChanged(this);
 	}
 }
 
diff --git a/Engine/Source/Editor/GameplayTasksEditor/Private/K2Node_LatentGameplayTaskCall.cpp b/Engine/Source/Editor/GameplayTasksEditor/Private/K2Node_LatentGameplayTaskCall.cpp
--- a/Engine/Source/Editor/GameplayTasksEditor/Private/K2Node_LatentGameplayTaskCall.cpp
+++ b/Engine/Source/Editor/GameplayTasksEditor/Private/K2Node_LatentGameplayTaskCall.cpp
@@ -288,7 +288,7 @@ void UK2Node_LatentGameplayTaskCall::PinDefaultValueChanged(UEdGraphPin* Changed
 
 		// Refresh the UI for the graph so the pin changes show up
 		UEdGraph* Graph = GetGraph();
-		Graph->NotifyGraphChanged();
+		Graph->NotifyNodeChanged(this);
 
 		// Mark dirty
 		FBlueprintEditorUtils::MarkBlueprintAsModified(GetBlueprint());
diff --git a/Engine/Source/Editor/GraphEditor/Private/KismetNodes/SGraphNodeK2Sequence.cpp b/Engine/Source/Editor/GraphEditor/Private/KismetNodes/SGraphNodeK2Sequence.cpp
--- a/Engine/Source/Editor/GraphEditor/Private/KismetNodes/SGraphNodeK2Sequence.cpp
+++ b/Engine/Source/Editor/GraphEditor/Private/KismetNodes/SGraphNodeK2Sequence.cpp
@@ -61,7 +61,7 @@ FReply SGraphNodeK2Sequence::OnAddPin()
 
 		AddPinNode->AddInputPin();
 		UpdateGraphNode();
-		GraphNode->GetGraph()->NotifyGraphChanged();
+		GraphNode->GetGraph()->NotifyNodeChanged(GraphNode);
 	}
 	
 	return FReply::Handled();
diff --git a/Engine/Source/Editor/GraphEditor/Private/KismetNodes/SGraphNodeSwitchStatement.cpp b/Engine/Source/Editor/GraphEditor/Private/KismetNodes/SGraphNodeSwitchStatement.cpp
--- a/Engine/Source/Editor/GraphEditor/Private/KismetNodes/SGraphNodeSwitchStatement.cpp
+++ b/Engine/Source/Editor/GraphEditor/Private/KismetNodes/SGraphNodeSwitchStatement.cpp
@@ -138,7 +138,7 @@ FReply SGraphNodeSwitchStatement::OnAddPin()
 	FBlueprintEditorUtils::MarkBlueprintAsModified(SwitchNode->GetBlueprint());
 
 	UpdateGraphNode();
-	GraphNode->GetGraph()->NotifyGraphChanged();
+	GraphNode->GetGraph()->NotifyNodeChanged(GraphNode);
 
 	return FReply::Handled();
 }
diff --git a/Engine/Source/Runtime/Engine/Classes/EdGraph/EdGraph.h b/Engine/Source/Runtime/Engine/Classes/EdGraph/EdGraph.h
--- a/Engine/Source/Runtime/Engine/Classes/EdGraph/EdGraph.h
+++ b/Engine/Source/Runtime/Engine/Classes/EdGraph/EdGraph.h
@@ -165,9 +165,12 @@ public:
 	*/
 	bool RemoveNode( UEdGraphNode* NodeToRemove, bool bBreakAllLinks = true );
 
-	/** Signal to listeners that the graph has changed */
+	/** Signal to listeners that the graph has changed - prefer to use NotifyNodeChanged when updating a single node */
 	virtual void NotifyGraphChanged();
 
+	/** Signal to listeners that a node has changed in the graph - commonly used to get UI up to date with a data change */
+	void NotifyNodeChanged(const UEdGraphNode* Node);
+
 	/** 
 	 * Move all nodes from this graph to another graph
 	 * @param DestinationGraph	The graph to move the nodes too 
diff --git a/Engine/Source/Runtime/Engine/Private/EdGraph/EdGraph.cpp b/Engine/Source/Runtime/Engine/Private/EdGraph/EdGraph.cpp
--- a/Engine/Source/Runtime/Engine/Private/EdGraph/EdGraph.cpp
+++ b/Engine/Source/Runtime/Engine/Private/EdGraph/EdGraph.cpp
@@ -283,6 +283,11 @@ void UEdGraph::NotifyGraphChanged()
 	OnGraphChanged.Broadcast(Action);
 }
 
+void UEdGraph::NotifyNodeChanged(const UEdGraphNode* Node)
+{
+	NotifyGraphChanged();
+}
+
 void UEdGraph::NotifyGraphChanged(const FEdGraphEditAction& InAction)
 {
 	OnGraphChanged.Broadcast(InAction);
