UE-97716 - Private and Protected functions do not display in the context menu in child BPs

We need the classes used by the target function and the graph to be consistent. Specifically, we could run into scenarios where the TargetFunction would be the skeleton class and *not* the generated class. This would cause subsequent IsChildOf usage in CanFunctionBeUsedInGraph to fail, since they're technically two different classes.

In this case, we can explicitly use the skeleton classes from both to remove any ambiguity.

diff --git a/Engine/Source/Editor/BlueprintGraph/Private/K2Node_CallFunction.cpp b/Engine/Source/Editor/BlueprintGraph/Private/K2Node_CallFunction.cpp
--- a/Engine/Source/Editor/BlueprintGraph/Private/K2Node_CallFunction.cpp
+++ b/Engine/Source/Editor/BlueprintGraph/Private/K2Node_CallFunction.cpp
@@ -1481,26 +1481,23 @@ bool UK2Node_CallFunction::CanPasteHere(const UEdGraph* TargetGraph) const
 	// Cannot paste editor only functions into runtime graphs
 	if (bCanPaste)
 	{
-		UFunction* TargetFunction = GetTargetFunction();
+		UFunction* TargetFunction = GetTargetFunctionFromSkeletonClass();
 
 		bCanPaste = CanEditorOnlyFunctionBeCalled(TargetFunction, TargetGraph);
 	}
 
 	// We check function context for placability only in the base class case; derived classes are typically bound to
 	// specific functions that should always be placeable, but may not always be explicitly callable (e.g. InternalUseOnly).
-	if(bCanPaste && GetClass() == StaticClass())
+	if (bCanPaste && GetClass() == StaticClass())
 	{
 		const UEdGraphSchema_K2* K2Schema = GetDefault<UEdGraphSchema_K2>();
 		uint32 AllowedFunctionTypes = UEdGraphSchema_K2::EFunctionType::FT_Pure | UEdGraphSchema_K2::EFunctionType::FT_Const | UEdGraphSchema_K2::EFunctionType::FT_Protected;
-		if(K2Schema->DoesGraphSupportImpureFunctions(TargetGraph))
+		if (K2Schema->DoesGraphSupportImpureFunctions(TargetGraph))
 		{
 			AllowedFunctionTypes |= UEdGraphSchema_K2::EFunctionType::FT_Imperative;
 		}
-		UFunction* TargetFunction = GetTargetFunction();
-		if( !TargetFunction )
-		{
-			TargetFunction = GetTargetFunctionFromSkeletonClass();
-		}
+		UFunction* TargetFunction = GetTargetFunctionFromSkeletonClass();
+
 		if (!TargetFunction)
 		{
 			// If the function doesn't exist and it is from self context, then it could be created from a CustomEvent node, that was also pasted (but wasn't compiled yet).
@@ -1508,7 +1505,7 @@ bool UK2Node_CallFunction::CanPasteHere(const UEdGraph* TargetGraph) const
 		}
 		else
 		{
-			bCanPaste = K2Schema->CanFunctionBeUsedInGraph(FBlueprintEditorUtils::FindBlueprintForGraphChecked(TargetGraph)->GeneratedClass, TargetFunction, TargetGraph, AllowedFunctionTypes, false);
+			bCanPaste = K2Schema->CanFunctionBeUsedInGraph(FBlueprintEditorUtils::FindBlueprintForGraphChecked(TargetGraph)->SkeletonGeneratedClass, TargetFunction, TargetGraph, AllowedFunctionTypes, false);
 		}
 	}
 	
