Make sure we use authoritative class, in case a literal object is being reinstanced on load

diff --git a/Engine/Source/Editor/BlueprintGraph/Private/K2Node_Literal.cpp b/Engine/Source/Editor/BlueprintGraph/Private/K2Node_Literal.cpp
--- a/Engine/Source/Editor/BlueprintGraph/Private/K2Node_Literal.cpp
+++ b/Engine/Source/Editor/BlueprintGraph/Private/K2Node_Literal.cpp
@@ -292,7 +292,7 @@ void UK2Node_Literal::SetObjectRef(UObject* NewValue)
 			ValuePin->Modify();
 			ValuePin->PinType.PinCategory = UEdGraphSchema_K2::PC_Object;
 			ValuePin->PinType.PinSubCategory = NAME_None;
-			ValuePin->PinType.PinSubCategoryObject = ObjectRef->GetClass();
+			ValuePin->PinType.PinSubCategoryObject = ObjectRef->GetClass()->GetAuthoritativeClass();
 		}
 	}
 
