Fixed text key determinism issue with AutogeneratedDefaultValue

ConformAutogeneratedDefaultValuePackage assumed that DefaultTextValue had already been initialized, but that doesn't seem to be the case during AllocateDefaultPins.

This would result in AutogeneratedDefaultValue failing to fix-up for the current package ID, and would produce text with a non-deterministic key.

diff --git a/Engine/Source/Editor/BlueprintGraph/Private/EdGraphSchema_K2.cpp b/Engine/Source/Editor/BlueprintGraph/Private/EdGraphSchema_K2.cpp
--- a/Engine/Source/Editor/BlueprintGraph/Private/EdGraphSchema_K2.cpp
+++ b/Engine/Source/Editor/BlueprintGraph/Private/EdGraphSchema_K2.cpp
@@ -5331,6 +5331,7 @@ void UEdGraphSchema_K2::SetPinAutogeneratedDefaultValueBasedOnType(UEdGraphPin*
 
 static void ConformAutogeneratedDefaultValuePackage(
 	const FEdGraphPinType& PinType, 
+	const UEdGraphNode* OwningNode, 
 	FString& AutogeneratedDefaultValue, 
 	const FText& DefaultTextValue
 )
@@ -5339,36 +5340,34 @@ static void ConformAutogeneratedDefaultValuePackage(
 		!PinType.IsContainer() &&
 		!AutogeneratedDefaultValue.IsEmpty())
 	{
-		FText TextValue;
-		const TCHAR* Success = FTextStringHelper::ReadFromBuffer(*AutogeneratedDefaultValue, TextValue);
-		check(Success);
-
-		const TOptional<FString> RealPackageNamespace = FTextInspector::GetNamespace(DefaultTextValue);
-		const TOptional<FString> SavedPackageNamespace = FTextInspector::GetNamespace(TextValue);
-		const TOptional<FString> CurrentKey = FTextInspector::GetKey(TextValue);
+		// Attempt to find the correct package namespace to use for text within this pin
+		// Favor using the node if we have it, as that will be most up-to-date
+		FString PackageNamespace;
+		if (OwningNode)
+		{
+			PackageNamespace = TextNamespaceUtil::GetPackageNamespace(OwningNode);
+		}
+		else if (!DefaultTextValue.IsEmpty())
+		{
+			PackageNamespace = TextNamespaceUtil::ExtractPackageNamespace(FTextInspector::GetNamespace(DefaultTextValue).Get(FString()));
+		}
 
-		if (RealPackageNamespace.IsSet() &&
-			CurrentKey.IsSet() &&
-			DefaultTextValue.ToString() == TextValue.ToString() &&
-			(	!SavedPackageNamespace.IsSet() ||
-				*RealPackageNamespace != *SavedPackageNamespace ) )
+		if (!PackageNamespace.IsEmpty())
 		{
-			const TOptional<FString> RealKey = FTextInspector::GetKey(DefaultTextValue);
+			FText AutogeneratedDefaultTextValue;
+			const TCHAR* Success = FTextStringHelper::ReadFromBuffer(*AutogeneratedDefaultValue, AutogeneratedDefaultTextValue);
+			check(Success);
 
-			// make new text value with same key in current namespace:
-			AutogeneratedDefaultValue.Empty();
-			if (RealKey.IsSet())
-			{
-				// maintain the key on DefaultTextValue:
-				TextValue = FText::ChangeKey(*RealPackageNamespace, *RealKey, TextValue);
-			}
-			else
+			// Conform the auto-generated default against this package ID, preserving its key to avoid determinism issues
+			const FText ConformedAutogeneratedDefaultTextValue = TextNamespaceUtil::CopyTextToPackage(AutogeneratedDefaultTextValue, PackageNamespace, TextNamespaceUtil::ETextCopyMethod::PreserveKey);
+
+			// IdenticalTo is a quick test for whether CopyTextToPackage returned the same text it was given (meaning there's nothing to update)
+			if (!ConformedAutogeneratedDefaultTextValue.IdenticalTo(AutogeneratedDefaultTextValue))
 			{
-				// reset back to the key in AutogeneratedDefaultValue - generating a new one
-				// will cause determinism issues which are extremely hard to track down:
-				TextValue = FText::ChangeKey(*RealPackageNamespace, *CurrentKey, TextValue);
+				// Fix-up the auto-generated default from the conformed value
+				AutogeneratedDefaultValue.Reset();
+				FTextStringHelper::WriteToBuffer(AutogeneratedDefaultValue, ConformedAutogeneratedDefaultTextValue);
 			}
-			FTextStringHelper::WriteToBuffer(AutogeneratedDefaultValue, TextValue);
 		}
 	}
 }
@@ -5383,7 +5382,7 @@ void UEdGraphSchema_K2::ResetPinToAutogeneratedDefaultValue(UEdGraphPin* Pin, bo
 	else
 	{
 		// Autogenerated value has unreliable package namespace for text value, hack fix it up now:
-		ConformAutogeneratedDefaultValuePackage(Pin->PinType, Pin->AutogeneratedDefaultValue, Pin->DefaultTextValue);
+		ConformAutogeneratedDefaultValuePackage(Pin->PinType, Pin->GetOwningNodeUnchecked(), Pin->AutogeneratedDefaultValue, Pin->DefaultTextValue);
 
 		GetPinDefaultValuesFromString(Pin->PinType, Pin->GetOwningNodeUnchecked(), Pin->AutogeneratedDefaultValue, Pin->DefaultValue, Pin->DefaultObject, Pin->DefaultTextValue, false);
 
