Enable non-blocking context menu mode by default in the Blueprint graph editor.

diff --git a/Engine/Source/Editor/BlueprintGraph/Private/BlueprintEditorSettings.cpp b/Engine/Source/Editor/BlueprintGraph/Private/BlueprintEditorSettings.cpp
--- a/Engine/Source/Editor/BlueprintGraph/Private/BlueprintEditorSettings.cpp
+++ b/Engine/Source/Editor/BlueprintGraph/Private/BlueprintEditorSettings.cpp
@@ -52,9 +52,9 @@ UBlueprintEditorSettings::UBlueprintEditorSettings(const FObjectInitializer& Obj
 	, BreakpointReloadMethod(EBlueprintBreakpointReloadMethod::RestoreAll)
 	, bEnablePinValueInspectionTooltips(true)
 	, bEnableNamespaceEditorFeatures(true)
-	// Experimental
-	, bEnableContextMenuTimeSlicing(false)
+	, bEnableContextMenuTimeSlicing(true)
 	, ContextMenuTimeSlicingThresholdMs(50)
+	// Experimental
 	, bFavorPureCastNodes(false)
 	, bDoNotMarkAllInstancesDirtyOnDefaultValueChange(false)
 	// Compiler Settings
diff --git a/Engine/Source/Editor/BlueprintGraph/Public/BlueprintEditorSettings.h b/Engine/Source/Editor/BlueprintGraph/Public/BlueprintEditorSettings.h
--- a/Engine/Source/Editor/BlueprintGraph/Public/BlueprintEditorSettings.h
+++ b/Engine/Source/Editor/BlueprintGraph/Public/BlueprintEditorSettings.h
@@ -196,16 +196,16 @@ public:
 	UPROPERTY(EditAnywhere, config, Category = Workflow, meta = (EditCondition = "bEnableNamespaceEditorFeatures", DisplayName = "Global Namespace Imports (Local User Only)"))
 	TArray<FString> NamespacesToAlwaysInclude;
 
-// Experimental
-public:
-	/** If enabled, when the Blueprint graph context menu is invoked (e.g. by dragging off a pin), it will not block the UI while it populates the available actions list. */
-	UPROPERTY(EditAnywhere, config, Category = Experimental, meta = (DisplayName = "Enable Non-Blocking Context Menu"))
+	/** When the Blueprint graph context menu is invoked (e.g. by right-clicking in the graph or dragging off a pin), do not block the UI while populating the available actions list. */
+	UPROPERTY(EditAnywhere, config, Category = Workflow, meta = (DisplayName = "Enable Non-Blocking Context Menu"))
 	bool bEnableContextMenuTimeSlicing;
 
 	/** The maximum amount of time (in milliseconds) allowed per frame for Blueprint graph context menu building when the non-blocking option is enabled. Larger values will complete the menu build in fewer frames, but will also make the UI feel less responsive. Smaller values will make the UI feel more responsive, but will also take longer to fully populate the menu. */
-	UPROPERTY(EditAnywhere, config, Category = Experimental, AdvancedDisplay, meta = (EditCondition = "bEnableContextMenuTimeSlicing", DisplayName = "Context Menu: Non-Blocking Per-Frame Threshold (ms)", ClampMin = "1"))
+	UPROPERTY(EditAnywhere, config, Category = Workflow, AdvancedDisplay, meta = (EditCondition = "bEnableContextMenuTimeSlicing", DisplayName = "Context Menu: Non-Blocking Per-Frame Threshold (ms)", ClampMin = "1"))
 	int32 ContextMenuTimeSlicingThresholdMs;
 
+// Experimental
+public:
 	/** If enabled, then placed cast nodes will default to their "pure" form (meaning: without execution pins). */
 	UPROPERTY(EditAnywhere, config, AdvancedDisplay, Category = Experimental, meta = (DisplayName = "Default to Using Pure Cast Nodes"))
 	bool bFavorPureCastNodes;
