[NO NON-PII COMMIT MESSAGE]

diff --git a/Engine/Source/Editor/BlueprintGraph/Private/K2Node_BaseAsyncTask.cpp b/Engine/Source/Editor/BlueprintGraph/Private/K2Node_BaseAsyncTask.cpp
--- a/Engine/Source/Editor/BlueprintGraph/Private/K2Node_BaseAsyncTask.cpp
+++ b/Engine/Source/Editor/BlueprintGraph/Private/K2Node_BaseAsyncTask.cpp
@@ -57,6 +57,7 @@ UK2Node_BaseAsyncTask::UK2Node_BaseAsyncTask(const FObjectInitializer& ObjectIni
 	, ProxyActivateFunctionName(NAME_None)
 	, bPinTooltipsValid(false)
 {
+	OrphanedPinSaveMode = ESaveOrphanPinMode::SaveAll;
 }
 
 FText UK2Node_BaseAsyncTask::GetTooltipText() const
