Default to enable the workflow option that will minimize what goes into the unsaved asset list after a Blueprint class default value change.

Notes:
- This also moves the option out of the "Experimental" section in Blueprint Editor settings.

diff --git a/Engine/Source/Editor/BlueprintGraph/Private/BlueprintEditorSettings.cpp b/Engine/Source/Editor/BlueprintGraph/Private/BlueprintEditorSettings.cpp
--- a/Engine/Source/Editor/BlueprintGraph/Private/BlueprintEditorSettings.cpp
+++ b/Engine/Source/Editor/BlueprintGraph/Private/BlueprintEditorSettings.cpp
@@ -55,9 +55,9 @@ UBlueprintEditorSettings::UBlueprintEditorSettings(const FObjectInitializer& Obj
 	, bEnableNamespaceEditorFeatures(true)
 	, bEnableContextMenuTimeSlicing(true)
 	, ContextMenuTimeSlicingThresholdMs(50)
+	, bDoNotMarkAllInstancesDirtyOnDefaultValueChange(true)
 	// Experimental
 	, bFavorPureCastNodes(false)
-	, bDoNotMarkAllInstancesDirtyOnDefaultValueChange(false)
 	// Compiler Settings
 	, SaveOnCompile(SoC_Never)
 	, bJumpToNodeErrors(false)
diff --git a/Engine/Source/Editor/BlueprintGraph/Public/BlueprintEditorSettings.h b/Engine/Source/Editor/BlueprintGraph/Public/BlueprintEditorSettings.h
--- a/Engine/Source/Editor/BlueprintGraph/Public/BlueprintEditorSettings.h
+++ b/Engine/Source/Editor/BlueprintGraph/Public/BlueprintEditorSettings.h
@@ -204,15 +204,15 @@ public:
 	UPROPERTY(EditAnywhere, config, Category = Workflow, AdvancedDisplay, meta = (EditCondition = "bEnableContextMenuTimeSlicing", DisplayName = "Context Menu: Non-Blocking Per-Frame Threshold (ms)", ClampMin = "1"))
 	int32 ContextMenuTimeSlicingThresholdMs;
 
+	/** If enabled, assets containing Blueprint instances (e.g. maps) will not be marked dirty when default values are edited, unless it results in the instance becoming realigned with the new default value. */
+	UPROPERTY(EditAnywhere, config, Category = Workflow)
+	bool bDoNotMarkAllInstancesDirtyOnDefaultValueChange;
+
 // Experimental
 public:
 	/** If enabled, then placed cast nodes will default to their "pure" form (meaning: without execution pins). */
 	UPROPERTY(EditAnywhere, config, AdvancedDisplay, Category = Experimental, meta = (DisplayName = "Default to Using Pure Cast Nodes"))
 	bool bFavorPureCastNodes;
-	
-	/** If enabled, assets containing Blueprint instances (e.g. maps) will not be marked dirty when default values are edited, unless it results in the instance becoming realigned with the new default value. */
-	UPROPERTY(EditAnywhere, config, Category = Experimental)
-	bool bDoNotMarkAllInstancesDirtyOnDefaultValueChange;
 
 	// Compiler Settings
 public:
