[CIS]: SA warning fix.

#rnx

diff --git a/Engine/Source/Editor/BlueprintGraph/Private/BlueprintEditorSettings.cpp b/Engine/Source/Editor/BlueprintGraph/Private/BlueprintEditorSettings.cpp
--- a/Engine/Source/Editor/BlueprintGraph/Private/BlueprintEditorSettings.cpp
+++ b/Engine/Source/Editor/BlueprintGraph/Private/BlueprintEditorSettings.cpp
@@ -285,14 +285,16 @@ bool UBlueprintEditorSettings::IsFunctionAllowed(const UBlueprint* InBlueprint,
 		return true;
 	}
 
-	const UClass* NativeParentClass = FBlueprintEditorUtils::FindFirstNativeClass(InBlueprint->ParentClass);
-	if (InBlueprint && NativeParentClass)
+	if (InBlueprint)
 	{
-		if (const UFunction* NativeParentFunction = NativeParentClass->FindFunctionByName(FunctionName))
+		if (const UClass* NativeParentClass = FBlueprintEditorUtils::FindFirstNativeClass(InBlueprint->ParentClass))
 		{
-			if (FunctionPermissions.PassesFilter(NativeParentFunction->GetPathName()))
+			if (const UFunction* NativeParentFunction = NativeParentClass->FindFunctionByName(FunctionName))
 			{
-				return true;
+				if (FunctionPermissions.PassesFilter(NativeParentFunction->GetPathName()))
+				{
+					return true;
+				}
 			}
 		}
 	}
