Allow users to control how assets selected in the Content Browser will generate pre-bound actions when invoking the Blueprint graph context menu. This feature is now also turned off by default.

Notes:
- This applies only when the "Context Sensitive" box is checked.
- New options are available in Editor Prefs->Blueprint Editor Settings:
    1) Whether actions pre-bound to selected assets will be included.
    2) Whether consideration will be restricted to single selections only.
    3) Whether selected assets will be loaded before being considered.
- Options 2 and 3 specify Option 1 as an edit condition in the markup.
- Default UX is 1-off, 2-off, 3-on. Customizing these options will affect the local user only.
- Legacy UX can be restored as 1-on, 2-off, 3-on if any licensees were relying on that specifically.

diff --git a/Engine/Source/Editor/BlueprintGraph/Private/BlueprintEditorSettings.cpp b/Engine/Source/Editor/BlueprintGraph/Private/BlueprintEditorSettings.cpp
--- a/Engine/Source/Editor/BlueprintGraph/Private/BlueprintEditorSettings.cpp
+++ b/Engine/Source/Editor/BlueprintGraph/Private/BlueprintEditorSettings.cpp
@@ -55,6 +55,9 @@ UBlueprintEditorSettings::UBlueprintEditorSettings(const FObjectInitializer& Obj
 	, bEnableNamespaceEditorFeatures(true)
 	, bEnableContextMenuTimeSlicing(true)
 	, ContextMenuTimeSlicingThresholdMs(50)
+	, bIncludeActionsForSelectedAssetsInContextMenu(false)
+	, bLimitAssetActionBindingToSingleSelectionOnly(false)
+	, bLoadSelectedAssetsForContextMenuActionBinding(true)
 	, bDoNotMarkAllInstancesDirtyOnDefaultValueChange(true)
 	// Experimental
 	, bFavorPureCastNodes(false)
diff --git a/Engine/Source/Editor/BlueprintGraph/Public/BlueprintEditorSettings.h b/Engine/Source/Editor/BlueprintGraph/Public/BlueprintEditorSettings.h
--- a/Engine/Source/Editor/BlueprintGraph/Public/BlueprintEditorSettings.h
+++ b/Engine/Source/Editor/BlueprintGraph/Public/BlueprintEditorSettings.h
@@ -204,6 +204,18 @@ public:
 	UPROPERTY(EditAnywhere, config, Category = Workflow, AdvancedDisplay, meta = (EditCondition = "bEnableContextMenuTimeSlicing", DisplayName = "Context Menu: Non-Blocking Per-Frame Threshold (ms)", ClampMin = "1"))
 	int32 ContextMenuTimeSlicingThresholdMs;
 
+	/** If enabled, invoking the Blueprint graph context menu with one or more compatible assets selected in the Content Browser will generate an additional set of pre-bound menu actions when the "Context Sensitive" option is enabled. For example, selecting a Static Mesh asset in the Content Browser will result in an extra "Add Static Mesh Component" menu action that's already bound to the selected asset. */
+	UPROPERTY(EditAnywhere, config, Category = Workflow, meta = (DisplayName = "Context Menu: Include Pre-Bound Actions for Selected Assets"))
+	bool bIncludeActionsForSelectedAssetsInContextMenu;
+
+	/** Only generate pre-bound "Add Component" actions when there is a single asset selected in the Content Browser. If more than one asset is selected, pre-bound "Add Component" actions will not be generated. Enabling this option can improve UI responsiveness and decrease the time it takes to build the context menu, while still preserving the ability to include actions pre-bound to the selected asset. */
+	UPROPERTY(EditAnywhere, config, Category = Workflow, meta = (EditCondition = "bIncludeActionsForSelectedAssetsInContextMenu", DisplayName = "Context Menu: Pre-Bound Actions: Restrict to Single Selection"))
+	bool bLimitAssetActionBindingToSingleSelectionOnly;
+
+	/** When generating pre-bound "Add Component" actions, any selected assets that are not yet loaded will be synchronously loaded as part of building the Blueprint Graph context menu. Enabling this option will ensure that all pre-bound actions for all selected assets are included in the menu, but load times may also affect editor UI responsiveness while the context menu is building. */
+	UPROPERTY(EditAnywhere, config, Category = Workflow, meta = (EditCondition = "bIncludeActionsForSelectedAssetsInContextMenu", DisplayName = "Context Menu: Pre-Bound Actions: Always Load Selected Asset(s)"))
+	bool bLoadSelectedAssetsForContextMenuActionBinding;
+
 	/** If enabled, assets containing Blueprint instances (e.g. maps) will not be marked dirty when default values are edited, unless it results in the instance becoming realigned with the new default value. */
 	UPROPERTY(EditAnywhere, config, Category = Workflow)
 	bool bDoNotMarkAllInstancesDirtyOnDefaultValueChange;
diff --git a/Engine/Source/Editor/Kismet/Private/BlueprintActionMenuUtils.cpp b/Engine/Source/Editor/Kismet/Private/BlueprintActionMenuUtils.cpp
--- a/Engine/Source/Editor/Kismet/Private/BlueprintActionMenuUtils.cpp
+++ b/Engine/Source/Editor/Kismet/Private/BlueprintActionMenuUtils.cpp
@@ -673,26 +673,32 @@ void FBlueprintActionMenuUtils::MakeContextMenu(FBlueprintActionContext const& C
 	AddComponentFilter.PermittedNodeTypes.Add(UK2Node_AddComponent::StaticClass());
 	AddComponentFilter.AddRejectionTest(FBlueprintActionFilter::FRejectionTestDelegate::CreateStatic(IsUnBoundSpawner));
 
-
-	FContentBrowserModule& ContentBrowserModule = FModuleManager::LoadModuleChecked<FContentBrowserModule>("ContentBrowser");
-	TArray<FAssetData> SelectedAssets;
-	ContentBrowserModule.Get().GetSelectedAssets(SelectedAssets);
-
-	for (FAssetData& Asset : SelectedAssets)
+	if (BlueprintSettings->bIncludeActionsForSelectedAssetsInContextMenu)
 	{
-		UClass* AssetClass = Asset.GetClass();
-		// filter here (rather than in FBlueprintActionFilter) so we only load
-		// assets that we can use
-		if ((AssetClass == nullptr) || (FComponentAssetBrokerage::GetPrimaryComponentForAsset(AssetClass) == nullptr))
+		FContentBrowserModule& ContentBrowserModule = FModuleManager::LoadModuleChecked<FContentBrowserModule>("ContentBrowser");
+		TArray<FAssetData> SelectedAssets;
+		ContentBrowserModule.Get().GetSelectedAssets(SelectedAssets);
+
+		if (SelectedAssets.Num() <= 1 || !BlueprintSettings->bLimitAssetActionBindingToSingleSelectionOnly)
 		{
-			continue;
-		}
+			for (FAssetData& Asset : SelectedAssets)
+			{
+				UClass* AssetClass = Asset.GetClass();
+				// filter here (rather than in FBlueprintActionFilter) so we only load
+				// assets that we can use
+				if ((AssetClass == nullptr) || (FComponentAssetBrokerage::GetPrimaryComponentForAsset(AssetClass) == nullptr))
+				{
+					continue;
+				}
 
-		// @TODO: loading assets here may be slow (but we need a UObject to 
-		//        properly bind to), consider adding a editor option that will 
-		//        only offer then if the asset is already loaded
-		UObject* AssetObj = Asset.GetAsset();
-		AddComponentFilter.Context.SelectedObjects.Add(AssetObj);
+				// loading assets here may be slow (but we need a UObject to properly bind to)
+				if (Asset.IsAssetLoaded() || BlueprintSettings->bLoadSelectedAssetsForContextMenuActionBinding)
+				{
+					UObject* AssetObj = Asset.GetAsset();
+					AddComponentFilter.Context.SelectedObjects.Add(AssetObj);
+				}
+			}
+		}
 	}
 
 	//--------------------------------------
