Go back to matching by name if there aren't conversion functions available to fix broken content cases where a name fix up does solve the issue (connected to wildcards, types change in tandem on multiple functions, pin object type changes to be more restrictive but connection is still valid).
#rnx

diff --git a/Engine/Source/Editor/BlueprintGraph/Private/K2Node.cpp b/Engine/Source/Editor/BlueprintGraph/Private/K2Node.cpp
--- a/Engine/Source/Editor/BlueprintGraph/Private/K2Node.cpp
+++ b/Engine/Source/Editor/BlueprintGraph/Private/K2Node.cpp
@@ -918,7 +918,7 @@ UK2Node::ERedirectType UK2Node::DoPinsMatchForReconstruction(const UEdGraphPin*
 				}
 				else
 				{
-					RedirectType = ERedirectType_None;
+					RedirectType = ERedirectType_Name;
 				}
 			}
 		}
