Include allowed properties/functions from an excluded class in the BP actions

#rnx

diff --git a/Engine/Source/Editor/BlueprintGraph/Private/BlueprintActionDatabase.cpp b/Engine/Source/Editor/BlueprintGraph/Private/BlueprintActionDatabase.cpp
--- a/Engine/Source/Editor/BlueprintGraph/Private/BlueprintActionDatabase.cpp
+++ b/Engine/Source/Editor/BlueprintGraph/Private/BlueprintActionDatabase.cpp
@@ -1571,15 +1571,10 @@ void FBlueprintActionDatabase::RefreshClassActions(UClass* const Class)
 {
 	TRACE_CPUPROFILER_EVENT_SCOPE(FBlueprintActionDatabase::RefreshClassActions);
 
-	// Early out if the class is filtered
-	if (!IsClassAllowed(Class, EPermissionsContext::Asset))
-	{
-		return;
-	}
-
 	using namespace BlueprintActionDatabaseImpl;
 	check(Class != nullptr);
 
+	bool const bFilterClass      = !IsClassAllowed(Class, EPermissionsContext::Asset);
 	bool const bOutOfDateClass   = Class->HasAnyClassFlags(CLASS_NewerVersionExists);
 	bool const bHiddenClass		 = Class->HasAnyClassFlags(CLASS_Hidden);
 	bool const bIsBlueprintClass = (Cast<UBlueprintGeneratedClass>(Class) != nullptr);
@@ -1592,6 +1587,12 @@ void FBlueprintActionDatabase::RefreshClassActions(UClass* const Class)
 	}
 	else if (bIsBlueprintClass)
 	{
+		// Early out if the class is filtered
+		if (bFilterClass)
+		{
+			return;
+		}
+
 		UBlueprint* Blueprint = Cast<UBlueprint>(Class->ClassGeneratedBy);
 		if ((Blueprint != nullptr) && BlueprintActionDatabaseImpl::IsObjectValidForDatabase(Blueprint))
 		{
@@ -1611,6 +1612,12 @@ void FBlueprintActionDatabase::RefreshClassActions(UClass* const Class)
 	// own actions (presumably ones that would spawn that node)...
 	else if (Class->IsChildOf<UEdGraphNode>())
 	{
+		// Early out if the class is filtered
+		if (bFilterClass)
+		{
+			return;
+		}
+
 		{
 			FActionList& ClassActionList = ActionRegistry.FindOrAdd(Class);
 			if (!bIsInitializing)
@@ -1646,13 +1653,19 @@ void FBlueprintActionDatabase::RefreshClassActions(UClass* const Class)
 	}
 	else if (Class->IsChildOf<UBlueprint>())
 	{
+		// Early out if the class is filtered
+		if (bFilterClass)
+		{
+			return;
+		}
+
 		FBlueprintActionDatabaseRegistrar Registrar(ActionRegistry, UnloadedActionRegistry, ActionPrimingQueue);
 		Cast<UBlueprint>(Class->ClassDefaultObject)->GetTypeActions(Registrar);
 	}
 	else
 	{
 		FActionList& ClassActionList = ActionRegistry.FindOrAdd(Class);
-		if (!bIsInitializing)
+		if (!bIsInitializing && !bFilterClass)
 		{
 			ClassActionList.Empty();
 			// if we're only refreshing this class (and not init'ing the whole 
@@ -1663,6 +1676,7 @@ void FBlueprintActionDatabase::RefreshClassActions(UClass* const Class)
 
 			RegisterAllNodeActions(Registrar);
 		}
+		// Note: We still run this if we're filtering the class, as the class itself may expose properties/functions/etc that derived non-filtered classes need access to
 		GetClassMemberActions(Class, ClassActionList);
 
 		// queue the newly added actions for priming
