Make sure we don't select skel class when accessing a sparse clas datamember that has been moved from a member variable

diff --git a/Engine/Source/Editor/BlueprintGraph/Private/K2Node_Variable.cpp b/Engine/Source/Editor/BlueprintGraph/Private/K2Node_Variable.cpp
--- a/Engine/Source/Editor/BlueprintGraph/Private/K2Node_Variable.cpp
+++ b/Engine/Source/Editor/BlueprintGraph/Private/K2Node_Variable.cpp
@@ -176,6 +176,16 @@ void UK2Node_Variable::CreatePinForSelf()
 					{
 						TargetClass = OwnerClass->GetAuthoritativeClass();
 					}
+					else if(TargetClass)
+					{
+						// check if it's a sparse member, sparse members are accessed via the authoritative
+						// class - this matches the convention defined in BlueprintActionDatabaseImpl::AddClassDataObjectActions:
+						UClass* AuthClass = TargetClass->GetAuthoritativeClass();
+						if (AuthClass->GetSparseClassDataStruct()->IsChildOf(Property->GetOwnerStruct()) )
+						{
+							TargetClass = AuthClass;
+						}
+					}
 				}
 				else if (GetBlueprint()->SkeletonGeneratedClass)
 				{
