UE-183502 - BP autoconversion adds extraneous conversion nodes

The previous attempt to add automatic conversion nodes (24029327) had a flawed design: checking pin connections during rewiring might have an incomplete view of the types of the linked pins.

For example, suppose we have two native, BlueprintCallable functions. One returns a float, and the other accepts a single float. In a Blueprint, we have the output of one linked to the input of the other.

Later, we update both functions to use ints. During rewiring of the function with the return value, it might see that its connection is a float (because that node hasn't been rewired yet). As a result, it'll insert a cast node. Subsequently, when we later rewire the function with the input, it'll see that it's connected to a float, because of the recently inserted cast node. This will add yet another cast node. Since both pin types are the same, there shouldn't be any cast nodes to begin with.

The only safe way to insert cast nodes is after node construction has finished. Instead of handling this during rewiring, we can defer the type checking to early validation. Overall, this simplifies the autoconversion since we just need to iterate over pairs of pins, check for type mismatches, and insert cast nodes where appropriate.

This change effectively reverts CLs 24174262, 24029327, 24218437, and 25444139, and moves the type checking logic to EarlyValidation.

diff --git a/Engine/Source/Editor/BlueprintGraph/Classes/K2Node.h b/Engine/Source/Editor/BlueprintGraph/Classes/K2Node.h
--- a/Engine/Source/Editor/BlueprintGraph/Classes/K2Node.h
+++ b/Engine/Source/Editor/BlueprintGraph/Classes/K2Node.h
@@ -27,6 +27,7 @@
 class AActor;
 class FArchive;
 class FBlueprintActionDatabaseRegistrar;
+class FCompilerResultsLog;
 class FKismetCompilerContext;
 class FProperty;
 class UActorComponent;
@@ -223,7 +224,7 @@ public:
 	BLUEPRINTGRAPH_API virtual bool CanSplitPin(const UEdGraphPin* Pin) const override;
 	BLUEPRINTGRAPH_API virtual UEdGraphPin* GetPassThroughPin(const UEdGraphPin* FromPin) const override;
 	BLUEPRINTGRAPH_API virtual bool IsInDevelopmentMode() const override;
-	BLUEPRINTGRAPH_API virtual void ValidateNodeDuringCompilation(class FCompilerResultsLog& MessageLog) const override;
+	BLUEPRINTGRAPH_API virtual void ValidateNodeDuringCompilation(FCompilerResultsLog& MessageLog) const override;
 	BLUEPRINTGRAPH_API virtual FString GetPinMetaData(FName InPinName, FName InKey) override;
 	// End of UEdGraphNode interface
 
@@ -372,10 +373,10 @@ public:
 	BLUEPRINTGRAPH_API virtual bool IsConnectionDisallowed(const UEdGraphPin* MyPin, const UEdGraphPin* OtherPin, FString& OutReason) const { return false; }
 
 	/** This function if used for nodes that needs CDO for validation (Called before expansion)*/
-	BLUEPRINTGRAPH_API virtual void EarlyValidation(class FCompilerResultsLog& MessageLog) const;
+	BLUEPRINTGRAPH_API virtual void EarlyValidation(FCompilerResultsLog& MessageLog) const;
 
 	/** This function returns an arbitrary number of attributes that describe this node for analytics events */
-	BLUEPRINTGRAPH_API virtual void GetNodeAttributes( TArray<TKeyValuePair<FString, FString>>& OutNodeAttributes ) const;
+	BLUEPRINTGRAPH_API virtual void GetNodeAttributes(TArray<TKeyValuePair<FString, FString>>& OutNodeAttributes) const;
 
 	/** Called before compilation begins, giving a blueprint time to force the linker to load data */
 	BLUEPRINTGRAPH_API virtual void PreloadRequiredAssets() { }
@@ -447,19 +448,14 @@ public:
 
 protected:
 
-	enum ERedirectType : uint8
+	enum ERedirectType
 	{
 		/* The pins do not match */
-		ERedirectType_None	= 0,
-
+		ERedirectType_None,
 		/* The pins match by name or redirect and have the same type (or we're ok with the mismatched type) */
-		ERedirectType_Name	= (1 << 0),
-
+		ERedirectType_Name,
 		/* The pins match via a redirect and the value needs to also be redirected */
-		ERedirectType_Value	= (1 << 1),
-
-		/* The pin types don't match, but can be connected via a cast node */
-		ERedirectType_Type	= (1 << 2),
+		ERedirectType_Value
 	};
 
 	// Handles the actual reconstruction (copying data, links, name, etc...) from two pins that have already been matched together
@@ -498,7 +494,7 @@ protected:
 	 * Should use this for node actions that happen during compilation!
 	 */
 	template<typename... ArgTypes>
-	void Message_Note(const FString& Message, ArgTypes... Args)
+	void Message_Note(const FString& Message, ArgTypes... Args) const
 	{
 		UBlueprint* OwningBP = GetBlueprint();
 		if (OwningBP)
@@ -512,7 +508,7 @@ protected:
 	}
 
 	template<typename... ArgTypes>
-	void Message_Warn(const FString& Message, ArgTypes... Args)
+	void Message_Warn(const FString& Message, ArgTypes... Args) const
 	{
 		UBlueprint* OwningBP = GetBlueprint();
 		if (OwningBP)
@@ -526,7 +522,7 @@ protected:
 	}
 
 	template<typename... ArgTypes>
-	void Message_Error(const FString& Message, ArgTypes... Args)
+	void Message_Error(const FString& Message, ArgTypes... Args) const
 	{
 		UBlueprint* OwningBP = GetBlueprint();
 		if (OwningBP)
@@ -564,7 +560,15 @@ private:
 	 * Utility function to write messages about orphan nodes in to the compiler log.
 	 * bStore indicates whether to write immediately to the log, or to store as a potential message to be committed once node pruning has completed
 	 */
-	void ValidateOrphanPins(class FCompilerResultsLog& MessageLog, bool bStore) const;
+	void ValidateOrphanPins(FCompilerResultsLog& MessageLog, bool bStore) const;
+	void ValidateOrphanPin(UEdGraphPin* Pin, FCompilerResultsLog& MessageLog, bool bStore) const;
+
+	/**
+	 * Checks that pins are type compatible with their links.
+	 * If there's a type mismatch (eg: due to a underlying type change), then we'll automatically
+	 * insert a conversion node into the graph if one exists between the two types.
+	 */
+	void ValidateLinkedPinTypes(UEdGraphPin* OutputPin, FCompilerResultsLog& MessageLog) const;
 
 public:
 
diff --git a/Engine/Source/Editor/BlueprintGraph/Classes/K2Node_CallFunction.h b/Engine/Source/Editor/BlueprintGraph/Classes/K2Node_CallFunction.h
--- a/Engine/Source/Editor/BlueprintGraph/Classes/K2Node_CallFunction.h
+++ b/Engine/Source/Editor/BlueprintGraph/Classes/K2Node_CallFunction.h
@@ -119,7 +119,6 @@ public:
 	// End of UEdGraphNode interface
 
 	// UK2Node interface
-	virtual ERedirectType DoPinsMatchForReconstruction(const UEdGraphPin* NewPin, int32 NewPinIndex, const UEdGraphPin* OldPin, int32 OldPinIndex) const override;
 	virtual void ReallocatePinsDuringReconstruction(TArray<UEdGraphPin*>& OldPins) override;
 	virtual bool IsNodePure() const override { return bIsPureFunc; }
 	virtual void PostReconstructNode() override;
diff --git a/Engine/Source/Editor/BlueprintGraph/Private/K2Node.cpp b/Engine/Source/Editor/BlueprintGraph/Private/K2Node.cpp
--- a/Engine/Source/Editor/BlueprintGraph/Private/K2Node.cpp
+++ b/Engine/Source/Editor/BlueprintGraph/Private/K2Node.cpp
@@ -886,41 +886,23 @@ void UK2Node::RestoreSplitPins(TArray<UEdGraphPin*>& OldPins)
 UK2Node::ERedirectType UK2Node::DoPinsMatchForReconstruction(const UEdGraphPin* NewPin, int32 NewPinIndex, const UEdGraphPin* OldPin, int32 OldPinIndex) const
 {
 	ERedirectType RedirectType = ERedirectType_None;
-	const UEdGraphSchema_K2* K2Schema = CastChecked<const UEdGraphSchema_K2>(GetSchema());
 
 	// if the pin names and directions do match
 	if (NewPin->PinName == OldPin->PinName && NewPin->Direction == OldPin->Direction)
 	{
-		const UClass* CallingContext = GetBlueprint()->GeneratedClass;
-		check(CallingContext);
-		const bool bPinsAreTypeCompatible = K2Schema->ArePinTypesCompatible(OldPin->PinType, NewPin->PinType, CallingContext);
-
 		// If the old pin had a default value, only match the new pin if the type is compatible:
-		if (OldPin->Direction == EGPD_Input &&
+		const UEdGraphSchema_K2* K2Schema = Cast<const UEdGraphSchema_K2>(GetSchema());
+		if (K2Schema && 
+			OldPin->Direction == EGPD_Input &&
 			OldPin->LinkedTo.Num() == 0 &&
 			!OldPin->DoesDefaultValueMatchAutogenerated() &&
-			!bPinsAreTypeCompatible)
+			!K2Schema->ArePinTypesCompatible(OldPin->PinType, NewPin->PinType))
 		{
 			RedirectType = ERedirectType_None;
 		}
 		else
 		{
-			if (bPinsAreTypeCompatible)
-			{
-				RedirectType = ERedirectType_Name;
-			}
-			else 
-			{
-				if (K2Schema->SearchForAutocastFunction(OldPin->PinType, NewPin->PinType).IsSet()
-					|| K2Schema->FindSpecializedConversionNode(OldPin->PinType, *NewPin, /* bCreateNode = */false).IsSet())
-				{
-					RedirectType = static_cast<ERedirectType>(static_cast<uint8>(ERedirectType_Name) | static_cast<uint8>(ERedirectType_Type));
-				}
-				else
-				{
-					RedirectType = ERedirectType_Name;
-				}
-			}
+			RedirectType = ERedirectType_Name;
 		}
 	}
 	else
@@ -997,16 +979,6 @@ UK2Node::ERedirectType UK2Node::DoPinsMatchForReconstruction(const UEdGraphPin*
 				// make sure they match
 				RedirectType = ((NewPin->PinName.ToString() != NewPinNameStr) ? ERedirectType_None : ERedirectType_Name);
 			}
-
-			// A type change might've also occurred with a rename.
-			if (RedirectType == ERedirectType_Name)
-			{
-				if (K2Schema->SearchForAutocastFunction(OldPin->PinType, NewPin->PinType).IsSet()
-					|| K2Schema->FindSpecializedConversionNode(OldPin->PinType, *NewPin, /* bCreateNode = */false).IsSet())
-				{
-					RedirectType = static_cast<ERedirectType>(static_cast<uint8>(ERedirectType_Name) | static_cast<uint8>(ERedirectType_Type));
-				}
-			}
 		}
 	}
 
@@ -1086,35 +1058,91 @@ void UK2Node::ReconstructSinglePin(UEdGraphPin* NewPin, UEdGraphPin* OldPin, ERe
 	}
 }
 
-void UK2Node::ValidateOrphanPins(FCompilerResultsLog& MessageLog, const bool bStore) const
+void UK2Node::ValidateOrphanPins(FCompilerResultsLog& MessageLog, bool bStore) const
 {
 	for (UEdGraphPin* Pin : Pins)
 	{
 		if (Pin && Pin->bOrphanedPin)
 		{
-			if (Pin->LinkedTo.Num())
-			{
-				const FText LinkedMessage = LOCTEXT("RemovedConnectedPin", "In use pin @@ no longer exists on node @@. Please refresh node or break links to remove pin.");
-				if (bStore)
-				{
-					MessageLog.StorePotentialError(this, *LinkedMessage.ToString(), Pin, this);
-				}
-				else
-				{
-					MessageLog.Error(*LinkedMessage.ToString(), Pin, this);
-				}
-			}
-			else if (!Pin->bHidden && !Pin->DoesDefaultValueMatchAutogenerated())
+			ValidateOrphanPin(Pin, MessageLog, bStore);
+		}
+	}
+}
+
+void UK2Node::ValidateOrphanPin(UEdGraphPin* Pin, FCompilerResultsLog& MessageLog, bool bStore) const
+{
+	check(Pin);
+	check(Pin->bOrphanedPin);
+
+	if (Pin->LinkedTo.Num())
+	{
+		const FText LinkedMessage = LOCTEXT("RemovedConnectedPin", "In use pin @@ no longer exists on node @@. Please refresh node or break links to remove pin.");
+		if (bStore)
+		{
+			MessageLog.StorePotentialError(this, *LinkedMessage.ToString(), Pin, this);
+		}
+		else
+		{
+			MessageLog.Error(*LinkedMessage.ToString(), Pin, this);
+		}
+	}
+	else if (!Pin->bHidden && !Pin->DoesDefaultValueMatchAutogenerated())
+	{
+		const FText NonDefaultMessage = LOCTEXT("RemovedNonDefaultPin", "Input pin @@ specifying non-default value no longer exists on node @@. Please refresh node or reset pin to default value to remove pin.");
+		if (bStore)
+		{
+			MessageLog.StorePotentialWarning(this, *NonDefaultMessage.ToString(), Pin, this);
+		}
+		else
+		{
+			MessageLog.Warning(*NonDefaultMessage.ToString(), Pin, this);
+		}
+	}
+}
+
+void UK2Node::ValidateLinkedPinTypes(UEdGraphPin* OutputPin, FCompilerResultsLog& MessageLog) const
+{
+	check(OutputPin);
+	check(OutputPin->Direction == EEdGraphPinDirection::EGPD_Output);
+
+	const UEdGraphSchema_K2* K2Schema = CastChecked<const UEdGraphSchema_K2>(GetSchema());
+
+	// Rewiring can mutate the linked pins arrays.
+	TArray<UEdGraphPin*> LinkedToPins = OutputPin->LinkedTo;
+
+	for (UEdGraphPin* InputPin : LinkedToPins)
+	{
+		check(InputPin);
+
+		// Function call nodes implicitly handle interface->object type redirects on self pins.
+		const bool bIsSelfInterfaceObjectConnection =
+			K2Schema->IsSelfPin(*InputPin) &&
+			(OutputPin->PinType.PinCategory == UEdGraphSchema_K2::PC_Interface) &&
+			(InputPin->PinType.PinCategory == UEdGraphSchema_K2::PC_Object)
+		;
+
+		const bool bNeedsConversionNode =
+			!bIsSelfInterfaceObjectConnection &&
+			!K2Schema->ArePinTypesCompatible(OutputPin->PinType, InputPin->PinType, GetBlueprintClassFromNode())
+		;
+
+		if (bNeedsConversionNode)
+		{
+			if (K2Schema->CreateAutomaticConversionNodeAndConnections(OutputPin, InputPin))
 			{
-				const FText NonDefaultMessage = LOCTEXT("RemovedNonDefaultPin", "Input pin @@ specifying non-default value no longer exists on node @@. Please refresh node or reset pin to default value to remove pin.");
-				if (bStore)
-				{
-					MessageLog.StorePotentialWarning(this, *NonDefaultMessage.ToString(), Pin, this);
-				}
-				else
-				{
-					MessageLog.Warning(*NonDefaultMessage.ToString(), Pin, this);
-				}
+				FFormatNamedArguments Args;
+
+				Args.Add(TEXT("PinName"), FText::FromName(OutputPin->GetFName()));
+				Args.Add(TEXT("LinkedPinName"), FText::FromName(InputPin->GetFName()));
+
+				FText ConversionInfo = FText::Format(
+					NSLOCTEXT("K2Node", "PinConversion_InfoFmt", "Pin '{PinName}' on node @@ doesn't match the type of '{LinkedPinName}' on node @@. A conversion node had to be inserted."),
+					Args
+				);
+
+				const UEdGraphPin* SourcePin = FBlueprintEditorUtils::FindFirstCompilerRelevantLinkedPin(OutputPin);
+				check(SourcePin);
+				Message_Note(ConversionInfo.ToString(), SourcePin->GetOwningNode(), InputPin->GetOwningNode());
 			}
 		}
 	}
@@ -1122,7 +1150,34 @@ void UK2Node::ValidateOrphanPins(FCompilerResultsLog& MessageLog, const bool bSt
 
 void UK2Node::EarlyValidation(FCompilerResultsLog& MessageLog) const
 {
-	ValidateOrphanPins(MessageLog, true);
+	UBlueprint* BP = GetBlueprint();
+	const UEdGraphSchema_K2* K2Schema = CastChecked<const UEdGraphSchema_K2>(GetSchema());
+
+	for (UEdGraphPin* Pin : Pins)
+	{
+		check(Pin);
+
+		// We only have to validate pin type compatability after the BP is first loaded.
+		// After the initial load, we can assume that the type won't change again.
+		// 
+		// Additionally, we can just check the output pins for type compatability.
+		// It would be redundant to also include checking input pins.
+		const bool bCheckPinForTypeChange =
+			BP->bIsRegeneratingOnLoad &&
+			!Pin->bOrphanedPin &&
+			!K2Schema->IsExecPin(*Pin) &&
+			(Pin->Direction == EEdGraphPinDirection::EGPD_Output)
+		;
+
+		if (Pin->bOrphanedPin)
+		{
+			ValidateOrphanPin(Pin, MessageLog, true);
+		}
+		else if (bCheckPinForTypeChange)
+		{
+			ValidateLinkedPinTypes(Pin, MessageLog);
+		}
+	}
 }
 
 void UK2Node::ValidateNodeDuringCompilation(FCompilerResultsLog& MessageLog) const
@@ -1201,10 +1256,7 @@ void UK2Node::RewireOldPinsToNewPins(TArray<UEdGraphPin*>& InOldPins, TArray<UEd
 				UEdGraphPin* NewPin = InNewPins[NewPinIndex];
 
 				const ERedirectType RedirectType = DoPinsMatchForReconstruction(NewPin, NewPinIndex, OldPin, OldPinIndex);
-				const bool bShouldReconstruct = (RedirectType != ERedirectType_None);
-				const bool bNeedsConversion = (RedirectType & ERedirectType_Type);
-
-				if (bShouldReconstruct)
+				if (RedirectType != ERedirectType_None)
 				{
 					ReconstructSinglePin(NewPin, OldPin, RedirectType);
 					MatchedPins.Add(OldPin, NewPin);
@@ -1214,47 +1266,6 @@ void UK2Node::RewireOldPinsToNewPins(TArray<UEdGraphPin*>& InOldPins, TArray<UEd
 					}
 					bMatched = true;
 					NewPinMatched[NewPinIndex] = true;
-
-					if (bNeedsConversion)
-					{
-						const UEdGraphSchema* Schema = GetSchema();
-						check(Schema);
-
-						TArray<UEdGraphPin*> LinkedToCopy = NewPin->LinkedTo;
-						for (UEdGraphPin* LinkedPin : LinkedToCopy)
-						{
-							// A linked pin could've had its type updated independently of ours.
-							// In some cases, it might've been changed to the same type, which alleviates the need for a cast.
-							if (NewPin->PinType != LinkedPin->PinType)
-							{
-								FFormatNamedArguments Args;
-
-								Args.Add(TEXT("PinName"), FText::FromName(NewPin->GetFName()));
-
-								if (Schema->CreateAutomaticConversionNodeAndConnections(NewPin, LinkedPin))
-								{
-									FText ConversionInfo = FText::Format(
-										NSLOCTEXT("K2Node", "PinConversion_InfoFmt", "Pin '{PinName}' on node @@ changed type. A conversion node had to be inserted into the graph."),
-										Args
-									);
-
-									Message_Note(ConversionInfo.ToString(), this);
-								}
-								else
-								{
-									Args.Add(TEXT("LinkName"), FText::FromName(LinkedPin->GetFName()));
-
-									FText ErrorInfo = FText::Format(
-										NSLOCTEXT("K2Node", "PinConversion_ErrorFmt", "Pin '{PinName}' on node @@ changed type, but no suitable conversion node could be found for link '{LinkName}'!"),
-										Args
-									);
-
-									Message_Error(ErrorInfo.ToString(), this);
-								}
-							}
-						}
-					}
-
 					break;
 				}
 			}
diff --git a/Engine/Source/Editor/BlueprintGraph/Private/K2Node_BreakStruct.cpp b/Engine/Source/Editor/BlueprintGraph/Private/K2Node_BreakStruct.cpp
--- a/Engine/Source/Editor/BlueprintGraph/Private/K2Node_BreakStruct.cpp
+++ b/Engine/Source/Editor/BlueprintGraph/Private/K2Node_BreakStruct.cpp
@@ -347,28 +347,7 @@ UK2Node::ERedirectType UK2Node_BreakStruct::DoPinsMatchForReconstruction(const U
 	ERedirectType Result = UK2Node::DoPinsMatchForReconstruction(NewPin, NewPinIndex, OldPin, OldPinIndex);
 	if ((ERedirectType_None == Result) && DoRenamedPinsMatch(NewPin, OldPin, true))
 	{
-		const UEdGraphSchema_K2* K2Schema = GetDefault<UEdGraphSchema_K2>();
-		check(K2Schema);
-		const UClass* CallingContext = GetBlueprint()->GeneratedClass;
-		check(CallingContext);
-
-		const bool bPinsAreTypeCompatible = K2Schema->ArePinTypesCompatible(OldPin->PinType, NewPin->PinType, CallingContext);
-
-		if (bPinsAreTypeCompatible)
-		{
-			Result = ERedirectType_Name;
-		}
-		else
-		{
-			if (K2Schema->SearchForAutocastFunction(OldPin->PinType, NewPin->PinType).IsSet())
-			{
-				Result = static_cast<ERedirectType>(static_cast<uint8>(ERedirectType_Name) | static_cast<uint8>(ERedirectType_Type));
-			}
-			else
-			{
-				Result = ERedirectType_None;
-			}
-		}
+		Result = ERedirectType_Name;
 	}
 
 	return Result;
diff --git a/Engine/Source/Editor/BlueprintGraph/Private/K2Node_CallFunction.cpp b/Engine/Source/Editor/BlueprintGraph/Private/K2Node_CallFunction.cpp
--- a/Engine/Source/Editor/BlueprintGraph/Private/K2Node_CallFunction.cpp
+++ b/Engine/Source/Editor/BlueprintGraph/Private/K2Node_CallFunction.cpp
@@ -767,27 +767,6 @@ UEdGraphPin* FindSelfPin(TArray<UEdGraphPin*>& Pins)
 	return NULL;
 }
 
-UK2Node::ERedirectType UK2Node_CallFunction::DoPinsMatchForReconstruction(const UEdGraphPin* NewPin, int32 NewPinIndex, const UEdGraphPin* OldPin, int32 OldPinIndex) const
-{
-	ERedirectType Result = Super::DoPinsMatchForReconstruction(NewPin, NewPinIndex, OldPin, OldPinIndex);
-	const bool bNeedsAutoConvert = (Result & ERedirectType_Type);
-	if (bNeedsAutoConvert)
-	{
-		// Function call nodes implicitly handle interface->object type redirects on self pins.
-		const UEdGraphSchema_K2* K2Schema = CastChecked<const UEdGraphSchema_K2>(GetSchema());
-		if (K2Schema->IsSelfPin(*OldPin)
-			&& OldPin->PinType.PinCategory == UEdGraphSchema_K2::PC_Interface
-			&& NewPin->PinType.PinCategory == UEdGraphSchema_K2::PC_Object)
-		{
-			// Remove the conversion flag as we aren't going to explicitly add one here; instead,
-			// we add it into the intermediate graph at expansion time (for backwards-compatibility).
-			Result = static_cast<ERedirectType>(static_cast<uint8>(Result) & ~static_cast<uint8>(ERedirectType_Type));
-		}
-	}
-
-	return Result;
-}
-
 void UK2Node_CallFunction::ReallocatePinsDuringReconstruction(TArray<UEdGraphPin*>& OldPins)
 {
 	// BEGIN TEMP
diff --git a/Engine/Source/Editor/BlueprintGraph/Private/K2Node_MakeStruct.cpp b/Engine/Source/Editor/BlueprintGraph/Private/K2Node_MakeStruct.cpp
--- a/Engine/Source/Editor/BlueprintGraph/Private/K2Node_MakeStruct.cpp
+++ b/Engine/Source/Editor/BlueprintGraph/Private/K2Node_MakeStruct.cpp
@@ -314,29 +314,8 @@ UK2Node::ERedirectType UK2Node_MakeStruct::DoPinsMatchForReconstruction(const UE
 	ERedirectType Result = UK2Node::DoPinsMatchForReconstruction(NewPin, NewPinIndex, OldPin, OldPinIndex);
 	if ((ERedirectType_None == Result) && DoRenamedPinsMatch(NewPin, OldPin, false))
 	{
-		const UEdGraphSchema_K2* K2Schema = GetDefault<UEdGraphSchema_K2>();
-		check(K2Schema);
-		const UClass* CallingContext = GetBlueprint()->GeneratedClass;
-		check(CallingContext);
-		const bool bPinsAreTypeCompatible = K2Schema->ArePinTypesCompatible(OldPin->PinType, NewPin->PinType, CallingContext);
-
-		if (bPinsAreTypeCompatible)
-		{
-			Result = ERedirectType_Name;
-		}
-		else
-		{
-			if (K2Schema->SearchForAutocastFunction(OldPin->PinType, NewPin->PinType).IsSet())
-			{
-				Result = static_cast<ERedirectType>(static_cast<uint8>(ERedirectType_Name) | static_cast<uint8>(ERedirectType_Type));
-			}
-			else
-			{
-				Result = ERedirectType_None;
-			}
-		}
+		Result = ERedirectType_Name;
 	}
-
 	return Result;
 }
 
diff --git a/Engine/Source/Editor/BlueprintGraph/Private/K2Node_StructOperation.cpp b/Engine/Source/Editor/BlueprintGraph/Private/K2Node_StructOperation.cpp
--- a/Engine/Source/Editor/BlueprintGraph/Private/K2Node_StructOperation.cpp
+++ b/Engine/Source/Editor/BlueprintGraph/Private/K2Node_StructOperation.cpp
@@ -75,8 +75,14 @@ bool UK2Node_StructOperation::DoRenamedPinsMatch(const UEdGraphPin* NewPin, cons
 	{
 		const EEdGraphPinDirection StructDirection = bStructInVariablesOut ? EGPD_Input : EGPD_Output;
 		const EEdGraphPinDirection VariablesDirection = bStructInVariablesOut ? EGPD_Output : EGPD_Input;
+		const UEdGraphSchema_K2* K2Schema = GetDefault<UEdGraphSchema_K2>(); 
+		const bool bCompatible = K2Schema && K2Schema->ArePinTypesCompatible(NewPin->PinType, OldPin->PinType);
 
-		if (StructDirection == OldPin->Direction)
+		if (!bCompatible)
+		{
+			return false;
+		}
+		else if (StructDirection == OldPin->Direction)
 		{
 			// Struct name was changed, which is fine
 			return true;
diff --git a/Engine/Source/Editor/BlueprintGraph/Private/K2Node_Variable.cpp b/Engine/Source/Editor/BlueprintGraph/Private/K2Node_Variable.cpp
--- a/Engine/Source/Editor/BlueprintGraph/Private/K2Node_Variable.cpp
+++ b/Engine/Source/Editor/BlueprintGraph/Private/K2Node_Variable.cpp
@@ -282,7 +282,7 @@ FString UK2Node_Variable::GetFindReferenceSearchString() const
 	return ResultSearchString;
 }
 
-UK2Node::ERedirectType UK2Node_Variable::DoPinsMatchForReconstruction( const UEdGraphPin* NewPin, int32 NewPinIndex, const UEdGraphPin* OldPin, int32 OldPinIndex ) const 
+UK2Node::ERedirectType UK2Node_Variable::DoPinsMatchForReconstruction(const UEdGraphPin* NewPin, int32 NewPinIndex, const UEdGraphPin* OldPin, int32 OldPinIndex) const 
 {
 	const UEdGraphSchema_K2* K2Schema = GetDefault<UEdGraphSchema_K2>();
 	if( OldPin->PinType.PinCategory == UEdGraphSchema_K2::PC_Exec )
@@ -353,10 +353,6 @@ UK2Node::ERedirectType UK2Node_Variable::DoPinsMatchForReconstruction( const UEd
 		{
 			return ERedirectType_Name;
 		}
-		else if (K2Schema->SearchForAutocastFunction(OldPin->PinType, NewPin->PinType))
-		{
-			return ERedirectType_Type;
-		}
 		else
 		{
 			const bool bNewPinIsObject = (NewPin->PinType.PinCategory == UEdGraphSchema_K2::PC_Object);
