Fix potential crash when `bFavorPureCastNodes` is enabled.

`K2Node_Message` attempted to create an impure cast node when expanded, but `SetPurity(false)` had no effect in this case because inside it `IsNodePure` returned `false` even if the default cast purity was `Pure`. It resulted in the creation of a pure node, causing a crash because the Message node attempted to work with its (non-existent) exec pins.

diff --git a/Engine/Source/Editor/BlueprintGraph/Classes/K2Node_DynamicCast.h b/Engine/Source/Editor/BlueprintGraph/Classes/K2Node_DynamicCast.h
--- a/Engine/Source/Editor/BlueprintGraph/Classes/K2Node_DynamicCast.h
+++ b/Engine/Source/Editor/BlueprintGraph/Classes/K2Node_DynamicCast.h
@@ -98,6 +98,9 @@ protected:
 
 	/** Update exec pins when converting from impure to pure. */
 	bool ReconnectPureExecPins(TArray<UEdGraphPin*>& OldPins);
+
+	/** Sets the pure state to Pure or Impure if it's using the default one */
+	void InitPureState();
 	
 	/** Constructing FText strings can be costly, so we cache the node's title */
 	FNodeTextCache CachedNodeTitle;
diff --git a/Engine/Source/Editor/BlueprintGraph/Private/K2Node_DynamicCast.cpp b/Engine/Source/Editor/BlueprintGraph/Private/K2Node_DynamicCast.cpp
--- a/Engine/Source/Editor/BlueprintGraph/Private/K2Node_DynamicCast.cpp
+++ b/Engine/Source/Editor/BlueprintGraph/Private/K2Node_DynamicCast.cpp
@@ -83,25 +83,7 @@ void UK2Node_DynamicCast::Serialize(FArchive& Ar)
 
 void UK2Node_DynamicCast::CreateExecPins()
 {
-	// Ensure the node is either pure or impure, based on current settings.
-	const UEdGraphSchema_K2* K2Schema = Cast<UEdGraphSchema_K2>(GetSchema());
-	check(K2Schema != nullptr);
-	if (!K2Schema->DoesGraphSupportImpureFunctions(GetGraph()))
-	{
-		PureState = EPureState::Pure;
-	}
-	else if (PureState == EPureState::UseDefault)
-	{
-		const UBlueprintEditorSettings* BlueprintSettings = GetDefault<UBlueprintEditorSettings>();
-		if (BlueprintSettings->bFavorPureCastNodes)
-		{
-			PureState = EPureState::Pure;
-		}
-		else
-		{
-			PureState = EPureState::Impure;
-		}
-	}
+	InitPureState();
 
 	const bool bIsNodePure = IsNodePure();
 	if (!bIsNodePure)
@@ -324,6 +306,8 @@ UEdGraphPin* UK2Node_DynamicCast::GetBoolSuccessPin() const
 
 void UK2Node_DynamicCast::SetPurity(bool bNewPurity)
 {
+	InitPureState();
+
 	if (bNewPurity != IsNodePure())
 	{
 		if (bNewPurity)
@@ -621,4 +605,32 @@ bool UK2Node_DynamicCast::IsActionFilteredOut(const FBlueprintActionFilter& Filt
 	return bIsFilteredOut;
 }
 
+void UK2Node_DynamicCast::InitPureState()
+{
+	if (PureState == EPureState::Pure)
+	{
+		return;
+	}
+
+	// Ensure the node is either pure or impure, based on current settings.
+	const UEdGraphSchema_K2* K2Schema = Cast<UEdGraphSchema_K2>(GetSchema());
+	check(K2Schema != nullptr);
+	if (!K2Schema->DoesGraphSupportImpureFunctions(GetGraph()))
+	{
+		PureState = EPureState::Pure;
+	}
+	else if (PureState == EPureState::UseDefault)
+	{
+		const UBlueprintEditorSettings* BlueprintSettings = GetDefault<UBlueprintEditorSettings>();
+		if (BlueprintSettings->bFavorPureCastNodes)
+		{
+			PureState = EPureState::Pure;
+		}
+		else
+		{
+			PureState = EPureState::Impure;
+		}
+	}
+}
+
 #undef LOCTEXT_NAMESPACE
