Automatically create ref terms for BP functions, this makes it easier to toggle a blueprint paramter to pass by reference

diff --git a/Engine/Source/Editor/BlueprintGraph/Private/EdGraphSchema_K2.cpp b/Engine/Source/Editor/BlueprintGraph/Private/EdGraphSchema_K2.cpp
--- a/Engine/Source/Editor/BlueprintGraph/Private/EdGraphSchema_K2.cpp
+++ b/Engine/Source/Editor/BlueprintGraph/Private/EdGraphSchema_K2.cpp
@@ -1085,6 +1085,21 @@ void UEdGraphSchema_K2::GetAutoEmitTermParameters(const UFunction* Function, TAr
 			}
 		}
 	}
+
+	// Allow any params that are blueprint defined to be autocreated:
+	if (!FBlueprintEditorUtils::IsNativeSignature(Function))
+	{
+		for (	TFieldIterator<FProperty> ParamIter(Function, EFieldIterationFlags::Default); 
+				ParamIter && (ParamIter->PropertyFlags & CPF_Parm); 
+				++ParamIter)
+		{
+			FProperty* Param = *ParamIter;
+			if(Param->HasAnyPropertyFlags(CPF_ReferenceParm))
+			{
+				AutoEmitParameterNames.Add(Param->GetName());
+			}
+		}
+	}
 }
 
 bool UEdGraphSchema_K2::FunctionHasParamOfType(const UFunction* InFunction, UEdGraph const* InGraph, const FEdGraphPinType& DesiredPinType, bool bWantOutput) const
diff --git a/Engine/Source/Editor/BlueprintGraph/Private/K2Node_CallFunction.cpp b/Engine/Source/Editor/BlueprintGraph/Private/K2Node_CallFunction.cpp
--- a/Engine/Source/Editor/BlueprintGraph/Private/K2Node_CallFunction.cpp
+++ b/Engine/Source/Editor/BlueprintGraph/Private/K2Node_CallFunction.cpp
@@ -2860,11 +2860,8 @@ void UK2Node_CallFunction::ExpandNode(class FKismetCompilerContext& CompilerCont
 		if ( Function )
 		{
 			TArray<FString> AutoCreateRefTermPinNames;
-			const bool bHasAutoCreateRefTerms = Function->HasMetaData(FBlueprintMetadata::MD_AutoCreateRefTerm);
-			if ( bHasAutoCreateRefTerms )
-			{
-				CompilerContext.GetSchema()->GetAutoEmitTermParameters(Function, AutoCreateRefTermPinNames);
-			}
+			CompilerContext.GetSchema()->GetAutoEmitTermParameters(Function, AutoCreateRefTermPinNames);
+			const bool bHasAutoCreateRefTerms = AutoCreateRefTermPinNames.Num() != 0;
 
 			for (UEdGraphPin* Pin : Pins)
 			{
