Fixed Blueprint type autoconversion for variable nodes.

Unlike other nodes, the variable node is the only that has a ArePinTypesCompatible check. This created a discrepancy: type changes on function nodes would preserve the link and create a compilation error, but type changes on variable nodes would orphan the pin.

For consistency, this removes the ArePinTypesCompatible check. More importantly, this also fixes auto type conversion, which will analyze the link and insert a cast node at a later stage. However, there's also some legacy cruft that only applied to object pins with mismatched types. This particular edge case has been preserved, but it might need follow up to see if it's still needed.

diff --git a/Engine/Source/Editor/BlueprintGraph/Private/K2Node_Variable.cpp b/Engine/Source/Editor/BlueprintGraph/Private/K2Node_Variable.cpp
--- a/Engine/Source/Editor/BlueprintGraph/Private/K2Node_Variable.cpp
+++ b/Engine/Source/Editor/BlueprintGraph/Private/K2Node_Variable.cpp
@@ -294,6 +294,23 @@ UK2Node::ERedirectType UK2Node_Variable::DoPinsMatchForReconstruction(const UEdG
 	const bool bCanMatchSelfs = bPinNamesMatch || ((OldPin->PinName == UEdGraphSchema_K2::PN_Self) == (NewPin->PinName == UEdGraphSchema_K2::PN_Self));
 	const bool bTheSameDirection = (NewPin->Direction == OldPin->Direction);
 
+	const bool bNewPinIsObject = (NewPin->PinType.PinCategory == UEdGraphSchema_K2::PC_Object);
+	const bool bNewPinIsInterface = (NewPin->PinType.PinCategory == UEdGraphSchema_K2::PC_Interface);
+	const bool bNewIsClass = (NewPin->PinType.PinCategory == UEdGraphSchema_K2::PC_Class);
+	const bool bNewPinIsObjectType = 
+		bNewPinIsObject ||
+		bNewPinIsInterface ||
+		bNewIsClass
+	;
+	
+	const FEdGraphPinType& InputType = (OldPin->Direction == EGPD_Output) ? OldPin->PinType : NewPin->PinType;
+	const FEdGraphPinType& OutputType = (OldPin->Direction == EGPD_Output) ? NewPin->PinType : OldPin->PinType;
+	
+	const bool bHasIncompatibleObjectType =
+		bNewPinIsObjectType &&
+		!K2Schema->ArePinTypesCompatible(OutputType, InputType)
+	;
+
 	if (bCanMatchSelfs && bTheSameDirection)
 	{
 		// the order that the PinTypes are passed to ArePinTypesCompatible() 
@@ -307,8 +324,6 @@ UK2Node::ERedirectType UK2Node_Variable::DoPinsMatchForReconstruction(const UEdG
 		// variable has been altered to be a sub-class ref (meaning we should 
 		// treat the NewPin as an output)... the opposite applies if the pins 
 		// are inputs
-		const FEdGraphPinType& InputType  = (OldPin->Direction == EGPD_Output) ? OldPin->PinType : NewPin->PinType;
-		const FEdGraphPinType& OutputType = (OldPin->Direction == EGPD_Output) ? NewPin->PinType : OldPin->PinType;
 
 		if (NewPin->ParentPin)
 		{
@@ -349,20 +364,14 @@ UK2Node::ERedirectType UK2Node_Variable::DoPinsMatchForReconstruction(const UEdG
 
 			return ERedirectType_Name;
 		}
-		else if (K2Schema->ArePinTypesCompatible(OutputType, InputType))
-		{
-			return ERedirectType_Name;
-		}
-		else
+		else if (bHasIncompatibleObjectType)
 		{
-			const bool bNewPinIsObject = (NewPin->PinType.PinCategory == UEdGraphSchema_K2::PC_Object);
-
 			// Special Case: If we had a pin match, and the class isn't loaded 
 			//               yet because of a cyclic dependency, temporarily 
 			//               cast away the const, and fix up.
 			if ( bPinNamesMatch &&
-				(bNewPinIsObject || (NewPin->PinType.PinCategory == UEdGraphSchema_K2::PC_Interface)) &&
-				(NewPin->PinType.PinSubCategoryObject == NULL) )
+				(bNewPinIsObject || bNewPinIsInterface) &&
+				(NewPin->PinType.PinSubCategoryObject == nullptr) )
 			{
 				// @TODO:  Fix this up to be less hacky
 				UBlueprintGeneratedClass* TypeClass = Cast<UBlueprintGeneratedClass>(OldPin->PinType.PinSubCategoryObject.Get());
@@ -403,7 +412,6 @@ UK2Node::ERedirectType UK2Node_Variable::DoPinsMatchForReconstruction(const UEdG
 			{
 				const UClass* PSCOClass = Cast<UClass>(OldPin->PinType.PinSubCategoryObject.Get());
 				const bool bOldIsBlueprint = PSCOClass && PSCOClass->IsChildOf(UBlueprint::StaticClass());
-				const bool bNewIsClass     = (NewPin->PinType.PinCategory == UEdGraphSchema_K2::PC_Class);
 				// Special Case: If we're migrating from old blueprint references 
 				//               to class references, allow pins to be reconnected if coerced
 				if (bNewIsClass && bOldIsBlueprint)
@@ -414,6 +422,14 @@ UK2Node::ERedirectType UK2Node_Variable::DoPinsMatchForReconstruction(const UEdG
 				}
 			}
 		}
+		else
+		{
+			// By default, we allow the redirect if direction and name match.
+			// The type may not match, but we defer that check to ValidateLinkedPinTypes,
+			// which attempts to address the incompatibility by inserting a conversion node.
+			// This behavior is consistent with how UK2Node operates.
+			return ERedirectType_Name;
+		}
 	}
 
 	return ERedirectType_None;
