UE-188291 - Blueprint autoconversions can't handle default values

Blueprint autoconversions normally work by checking pairs of linked pins and inserting cast nodes. However, default literal values on nodes aren't linked, so they must be treated differently.

Since default values aren't linked, we need to check the old and new pin types during rewiring. We can't do this during validation stage since we've forgotten the old pin type by that point. Additionally, we should prevent any sort of data loss when handling the default values. The safest way to do so is:

1. Create a literal node from the old type.
2. Create a conversion node between the old and new type.
3. Link the nodes to the pin that has the default value.

The drawback to this technique is that it inserts nodes, but it's safer than directly modifying the default value on the new pin (eg: "1.1" could end up as "1"). We still provide a note that nodes were inserted, which gives the developer the choice of how to handle the conversion change.

diff --git a/Engine/Source/Editor/BlueprintGraph/Classes/K2Node.h b/Engine/Source/Editor/BlueprintGraph/Classes/K2Node.h
--- a/Engine/Source/Editor/BlueprintGraph/Classes/K2Node.h
+++ b/Engine/Source/Editor/BlueprintGraph/Classes/K2Node.h
@@ -455,7 +455,9 @@ protected:
 		/* The pins match by name or redirect and have the same type (or we're ok with the mismatched type) */
 		ERedirectType_Name,
 		/* The pins match via a redirect and the value needs to also be redirected */
-		ERedirectType_Value
+		ERedirectType_Value,
+		/* The pins differ by type and have a default value*/
+		ERedirectType_DefaultValue,
 	};
 
 	// Handles the actual reconstruction (copying data, links, name, etc...) from two pins that have already been matched together
@@ -570,6 +572,16 @@ private:
 	 */
 	void ValidateLinkedPinTypes(UEdGraphPin* OutputPin, FCompilerResultsLog& MessageLog) const;
 
+	/**
+	 * If there's a type mismatch between old and new pins during rewiring *and* the old pin has a default value,
+	 * then we'll need a way to keep the old default value and convert it to the new type.
+	 * This function achieves that by inserting both a literal and a conversion node.
+	 * The literal node will use the original default value.
+	 * 
+	 * Returns true if the conversion was successful.
+	 */
+	bool TryInsertDefaultValueConversionNode(const UEdGraphPin& OldPin, UEdGraphPin& NewPin) const;
+
 public:
 
 	/** Util to get the generated class from a node. */
diff --git a/Engine/Source/Editor/BlueprintGraph/Private/K2Node.cpp b/Engine/Source/Editor/BlueprintGraph/Private/K2Node.cpp
--- a/Engine/Source/Editor/BlueprintGraph/Private/K2Node.cpp
+++ b/Engine/Source/Editor/BlueprintGraph/Private/K2Node.cpp
@@ -13,8 +13,10 @@
 #include "GraphEditorSettings.h"
 #include "EdGraph/EdGraphSchema.h"
 #include "EdGraphSchema_K2.h"
+#include "EdGraphSchema_K2_Actions.h"
 #include "K2Node_CallFunction.h"
 #include "K2Node_MacroInstance.h"
+#include "Kismet/KismetSystemLibrary.h"
 #include "Kismet2/BlueprintEditorUtils.h"
 #include "Editor/EditorEngine.h"
 #include "Misc/OutputDeviceNull.h"
@@ -898,7 +900,7 @@ UK2Node::ERedirectType UK2Node::DoPinsMatchForReconstruction(const UEdGraphPin*
 			!OldPin->DoesDefaultValueMatchAutogenerated() &&
 			!K2Schema->ArePinTypesCompatible(OldPin->PinType, NewPin->PinType))
 		{
-			RedirectType = ERedirectType_None;
+			RedirectType = ERedirectType_DefaultValue;
 		}
 		else
 		{
@@ -1162,6 +1164,81 @@ void UK2Node::ValidateLinkedPinTypes(UEdGraphPin* OutputPin, FCompilerResultsLog
 	}
 }
 
+bool UK2Node::TryInsertDefaultValueConversionNode(const UEdGraphPin& OldPin, UEdGraphPin& NewPin) const
+{
+	static TMap<FName, FName> MakeLiteralFunctionTable =
+	{
+		{ UEdGraphSchema_K2::PC_Boolean,	TEXT("MakeLiteralBool")		},
+		{ UEdGraphSchema_K2::PC_Byte,		TEXT("MakeLiteralByte")		},
+		{ UEdGraphSchema_K2::PC_Int,		TEXT("MakeLiteralInt")		},
+		{ UEdGraphSchema_K2::PC_Int64,		TEXT("MakeLiteralInt64")	},
+		{ UEdGraphSchema_K2::PC_Name,		TEXT("MakeLiteralName")		},
+		{ UEdGraphSchema_K2::PC_Real,		TEXT("MakeLiteralDouble")	},
+		{ UEdGraphSchema_K2::PC_String,		TEXT("MakeLiteralString")	},
+		{ UEdGraphSchema_K2::PC_Text,		TEXT("MakeLiteralText")		},
+	};
+
+	const UEdGraphSchema_K2* K2Schema = CastChecked<UEdGraphSchema_K2>(GetSchema());
+
+	const bool bHasCastFunction = K2Schema->SearchForAutocastFunction(OldPin.PinType, NewPin.PinType).IsSet();
+	const FName* MakeLiteralFunctionName = MakeLiteralFunctionTable.Find(OldPin.PinType.PinCategory);
+
+	const bool bCanConvertDefaultValue =
+		bHasCastFunction &&
+		(MakeLiteralFunctionName != nullptr)
+		;
+
+	if (bCanConvertDefaultValue)
+	{
+		UFunction* MakeLiteralFunction = UKismetSystemLibrary::StaticClass()->FindFunctionByName(*MakeLiteralFunctionName);
+		check(MakeLiteralFunction);
+
+		const int32 NodeXOffset = -600;
+		FVector2D SpawnLocation(NodePosX + NodeXOffset, NodePosY);
+
+		UK2Node_CallFunction* MakeLiteralNode = FEdGraphSchemaAction_K2NewNode::SpawnNode<UK2Node_CallFunction>(
+			GetGraph(),
+			SpawnLocation,
+			EK2NewNodeFlags::None,
+			[MakeLiteralFunction](UK2Node_CallFunction* NewInstance)
+			{
+				NewInstance->SetFromFunction(MakeLiteralFunction);
+			}
+		);
+
+		UEdGraphPin* InputPin = MakeLiteralNode->FindPinByPredicate(
+			[K2Schema](UEdGraphPin* InPin)
+			{
+				check(InPin);
+				return (InPin->Direction == EGPD_Input) && !K2Schema->IsMetaPin(*InPin);
+			}
+		);
+		check(InputPin);
+		InputPin->DefaultValue = OldPin.DefaultValue;
+
+		UEdGraphPin* OutputPin = MakeLiteralNode->GetReturnValuePin();
+		check(OutputPin);
+
+		if (ensure(K2Schema->CreateAutomaticConversionNodeAndConnections(OutputPin, &NewPin)))
+		{
+			FFormatNamedArguments Args;
+
+			Args.Add(TEXT("PinName"), FText::FromName(OutputPin->GetFName()));
+
+			FText ConversionInfo = FText::Format(
+				NSLOCTEXT("K2Node", "PinDefaultConversion_InfoFmt", "Pin '{PinName}' on node @@ has a default value that doesn't match the original type. A literal node and a conversion node had to be inserted."),
+				Args
+			);
+
+			const UEdGraphPin* SourcePin = FBlueprintEditorUtils::FindFirstCompilerRelevantLinkedPin(OutputPin);
+			check(SourcePin);
+			Message_Note(ConversionInfo.ToString(), this);
+		}
+	}
+
+	return bCanConvertDefaultValue;
+}
+
 void UK2Node::EarlyValidation(FCompilerResultsLog& MessageLog) const
 {
 	UBlueprint* BP = GetBlueprint();
@@ -1269,7 +1346,16 @@ void UK2Node::RewireOldPinsToNewPins(TArray<UEdGraphPin*>& InOldPins, TArray<UEd
 			{
 				UEdGraphPin* NewPin = InNewPins[NewPinIndex];
 
-				const ERedirectType RedirectType = DoPinsMatchForReconstruction(NewPin, NewPinIndex, OldPin, OldPinIndex);
+				ERedirectType RedirectType = DoPinsMatchForReconstruction(NewPin, NewPinIndex, OldPin, OldPinIndex);
+
+				if (RedirectType == ERedirectType_DefaultValue)
+				{
+					if (!TryInsertDefaultValueConversionNode(*OldPin, *NewPin))
+					{
+						RedirectType = ERedirectType_None;
+					}
+				}
+
 				if (RedirectType != ERedirectType_None)
 				{
 					ReconstructSinglePin(NewPin, OldPin, RedirectType);
