Fixed const-correctness of `UK2Node_Event::FindEventSignatureFunction`

diff --git a/Engine/Source/Editor/BlueprintGraph/Classes/K2Node_Event.h b/Engine/Source/Editor/BlueprintGraph/Classes/K2Node_Event.h
--- a/Engine/Source/Editor/BlueprintGraph/Classes/K2Node_Event.h
+++ b/Engine/Source/Editor/BlueprintGraph/Classes/K2Node_Event.h
@@ -117,7 +117,7 @@ class UK2Node_Event : public UK2Node_EditablePinBase, public IK2Node_EventNodeIn
 	/** Checks if this event node is implementing an interface event */
 	BLUEPRINTGRAPH_API bool IsInterfaceEventNode() const;
 
-	BLUEPRINTGRAPH_API UFunction* FindEventSignatureFunction();
+	BLUEPRINTGRAPH_API UFunction* FindEventSignatureFunction() const;
 	BLUEPRINTGRAPH_API void UpdateDelegatePin(bool bSilent = false);
 	BLUEPRINTGRAPH_API FName GetFunctionName() const;
 	BLUEPRINTGRAPH_API virtual bool IsUsedByAuthorityOnlyDelegate() const { return false; }
diff --git a/Engine/Source/Editor/BlueprintGraph/Private/K2Node_Event.cpp b/Engine/Source/Editor/BlueprintGraph/Private/K2Node_Event.cpp
--- a/Engine/Source/Editor/BlueprintGraph/Private/K2Node_Event.cpp
+++ b/Engine/Source/Editor/BlueprintGraph/Private/K2Node_Event.cpp
@@ -336,7 +336,7 @@ FName UK2Node_Event::GetFunctionName() const
 	return bOverrideFunction ? EventReference.GetMemberName() : CustomFunctionName;
 }
 
-UFunction* UK2Node_Event::FindEventSignatureFunction()
+UFunction* UK2Node_Event::FindEventSignatureFunction() const
 {
 	return EventReference.ResolveMember<UFunction>(GetBlueprintClassFromNode());
 }
