Add an option to convert a comparison operator node to another comparison, like == to !=.

Add the symbol to the user facing operator name for comparison operators for a nicer menu

diff --git a/Engine/Source/Editor/BlueprintGraph/Classes/K2Node_PromotableOperator.h b/Engine/Source/Editor/BlueprintGraph/Classes/K2Node_PromotableOperator.h
--- a/Engine/Source/Editor/BlueprintGraph/Classes/K2Node_PromotableOperator.h
+++ b/Engine/Source/Editor/BlueprintGraph/Classes/K2Node_PromotableOperator.h
@@ -198,4 +198,15 @@ public:
 	* @return bool		True if a conversion is possible
 	*/
 	BLUEPRINTGRAPH_API bool CanConvertPinType(const UEdGraphPin* Pin) const;
+
+	/**
+	 * Returns true if we can convert this node to another comparison operator type.
+	 * This will only be true if the current operator is a comparison operator such as ==, !=, <, etc.
+	 */
+	static bool CanConvertComparisonOperatorNodeType(const UEdGraphNode* Node);
+
+	/**
+	 * Converts the given node to use the new operator name
+	 */
+	static void ConvertComparisonOperatorNode(UEdGraphNode* Node, const FName NewOpName);
 };
diff --git a/Engine/Source/Editor/BlueprintGraph/Private/BlueprintTypePromotion.cpp b/Engine/Source/Editor/BlueprintGraph/Private/BlueprintTypePromotion.cpp
--- a/Engine/Source/Editor/BlueprintGraph/Private/BlueprintTypePromotion.cpp
+++ b/Engine/Source/Editor/BlueprintGraph/Private/BlueprintTypePromotion.cpp
@@ -453,12 +453,12 @@ const FText& FTypePromotion::GetUserFacingOperatorName(const FName Operator)
 		{ OperatorNames::Multiply,  LOCTEXT("MultiplyDisplayName",		"Multiply") },
 		{ OperatorNames::Subtract,  LOCTEXT("SubtractDisplayName",		"Subtract") },
 		{ OperatorNames::Divide,	LOCTEXT("DivideDisplayName",		"Divide") },
-		{ OperatorNames::Greater,	LOCTEXT("GreaterDisplayName",		"Greater") },
-		{ OperatorNames::GreaterEq, LOCTEXT("GreaterEqDisplayName",		"Greater Equal") },
-		{ OperatorNames::Less,		LOCTEXT("LessDisplayName",			"Less") },
-		{ OperatorNames::LessEq,	LOCTEXT("LessEqDisplayName",		"Less Equal") },
-		{ OperatorNames::NotEq,		LOCTEXT("NotEqDisplayName",			"Not Equal") },
-		{ OperatorNames::Equal,		LOCTEXT("EqualDisplayName",			"Equal") },
+		{ OperatorNames::Greater,	LOCTEXT("GreaterDisplayName",		"Greater ( > )") },
+		{ OperatorNames::GreaterEq, LOCTEXT("GreaterEqDisplayName",		"Greater Equal ( >= )") },
+		{ OperatorNames::Less,		LOCTEXT("LessDisplayName",			"Less ( < )") },
+		{ OperatorNames::LessEq,	LOCTEXT("LessEqDisplayName",		"Less Equal ( <= )") },
+		{ OperatorNames::NotEq,		LOCTEXT("NotEqDisplayName",			"Not Equal ( != )") },
+		{ OperatorNames::Equal,		LOCTEXT("EqualDisplayName",			"Equal ( == )") },
 		{ OperatorNames::NoOp,		LOCTEXT("NoOpDisplayName",			"") },
 	};
 
diff --git a/Engine/Source/Editor/BlueprintGraph/Private/K2Node_PromotableOperator.cpp b/Engine/Source/Editor/BlueprintGraph/Private/K2Node_PromotableOperator.cpp
--- a/Engine/Source/Editor/BlueprintGraph/Private/K2Node_PromotableOperator.cpp
+++ b/Engine/Source/Editor/BlueprintGraph/Private/K2Node_PromotableOperator.cpp
@@ -45,6 +45,12 @@
 
 #define LOCTEXT_NAMESPACE "PromotableOperatorNode"
 
+static TAutoConsoleVariable<bool> CVarAllowConversionOfComparisonOps(
+	TEXT("BP.bAllowConversionOfComparisonOps"),
+	true,
+	TEXT("If true, then allow the user to convert between comparison operators on the UK2Node_PromotableOperator"),
+	ECVF_Default);
+
 ///////////////////////////////////////////////////////////
 // Pin names for default construction
 
@@ -138,7 +144,39 @@ void UK2Node_PromotableOperator::GetNodeContextMenuActions(UToolMenu* Menu, UGra
 
 	static const FName PromotableOperatorNodeName = FName("PromotableOperator");
 	static const FText PromotableOperatorStr = LOCTEXT("PromotableOperatorNode", "Operator Node");
+	
+	if (CVarAllowConversionOfComparisonOps.GetValueOnAnyThread() && CanConvertComparisonOperatorNodeType(Context->Node))
+	{
+		static const FName ConvertComparisonOpName = FName("ConvertComparisonOp");
+		static const FText ConvertComparisonOpStr = LOCTEXT("ConvertComparisonOp", "Convert Operator");
+		
+		FToolMenuSection& Section = Menu->AddSection(ConvertComparisonOpName, ConvertComparisonOpStr);
+		for (const FName& PossibleConversionOpName : FTypePromotion::GetComparisonOpNames())
+		{
+			// Don't display our current operator type
+			if (PossibleConversionOpName == OperationName)
+			{
+				continue;
+			}
+
+			FFormatNamedArguments Args;
+			Args.Add(TEXT("NewOpName"), FTypePromotion::GetUserFacingOperatorName(PossibleConversionOpName));
+			Args.Add(TEXT("CurrentOpName"), FTypePromotion::GetUserFacingOperatorName(OperationName));
 
+			const FText PinConversionName = FText::Format(LOCTEXT("CallFunction_Tooltip", "Convert to {NewOpName}"), Args);
+
+			Section.AddMenuEntry(
+				FName(PinConversionName.ToString()),
+				PinConversionName,
+				FText::Format(LOCTEXT("ConvertPinTypeTooltip", "Convert this node operation from '{CurrentOpName}' to '{NewOpName}'"), Args),
+				FSlateIcon(),
+				FUIAction(
+					FExecuteAction::CreateStatic(&UK2Node_PromotableOperator::ConvertComparisonOperatorNode, const_cast<UEdGraphNode*>(Context->Node.Get()), PossibleConversionOpName)
+				)
+			);
+		}
+	}
+	
 	// Add the option to remove a pin via the context menu
 	if (CanRemovePin(Context->Pin))
 	{
@@ -268,6 +306,68 @@ bool UK2Node_PromotableOperator::CanConvertPinType(const UEdGraphPin* Pin) const
 		!IsTolerancePin(*Pin);
 }
 
+bool UK2Node_PromotableOperator::CanConvertComparisonOperatorNodeType(const UEdGraphNode* Node)
+{
+	if (!Node)
+	{
+		return false;
+	}
+	
+	if (const UK2Node_PromotableOperator* OpNode = Cast<UK2Node_PromotableOperator>(Node))
+	{
+		// Only allow conversions between simple comparison types. If we have a tolerance pin, don't allow conversion
+		// because not all operators for those types have a tolerance.
+		const UEdGraphPin* TolerancePin = OpNode->FindTolerancePin();
+		const bool bTolerancePinAcceptable = !TolerancePin || TolerancePin->bHidden;
+		
+		const UFunction* TargetFunc = OpNode->GetTargetFunction();
+		
+		// If the target function is null then it's a wildcard node and we can convert it
+		return (TargetFunc == nullptr || FTypePromotion::IsComparisonFunc(TargetFunc)) && bTolerancePinAcceptable;
+	}
+
+	return false;
+}
+
+void UK2Node_PromotableOperator::ConvertComparisonOperatorNode(UEdGraphNode* Node, const FName NewOpName)
+{
+	UK2Node_PromotableOperator* OpNode = Cast<UK2Node_PromotableOperator>(Node);
+	// You can only convert to comparison operators, nothing else, because then we can be sure the number of pins to consider will be correct
+	// when finding the best matching function. Other math ops may behave in undefined ways if we allowed it.
+	if (!OpNode || !FTypePromotion::IsComparisonOpName(NewOpName))
+	{
+		return;
+	}
+	
+	FFormatNamedArguments Args;
+	Args.Add(TEXT("NewOpName"), FText::FromName(NewOpName));
+	Args.Add(TEXT("CurrentOpName"), FText::FromName(OpNode->OperationName));
+
+	const FText PinConversionName = FText::Format(LOCTEXT("CallFunction_Tooltip", "Convert operator node from {CurrentOpName} to {NewOpName}"), Args);
+	
+	FScopedTransaction Transaction(LOCTEXT("PromotableOperatorComparisonOpConversion", "Convert operator node"));
+	OpNode->Modify();
+	
+	OpNode->OperationName = NewOpName;
+
+	TArray<UEdGraphPin*> PinsToConsider;
+	OpNode->GetPinsToConsider(PinsToConsider);
+
+	if (const UFunction* BestMatchingFunc = FTypePromotion::FindBestMatchingFunc(OpNode->OperationName, PinsToConsider))
+	{
+		// Only allow this with comparison functions
+		ensure(FTypePromotion::IsComparisonFunc(BestMatchingFunc));
+		UE_LOG(LogBlueprint, Verbose, TEXT("Converting node '%s' from '%s' to '%s'..."), *GetNameSafe(OpNode), *OpNode->OperationName.ToString(), *NewOpName.ToString());
+		
+		OpNode->SetFromFunction(BestMatchingFunc);
+		OpNode->ReconstructNode();
+	}
+	else
+	{
+		UE_LOG(LogBlueprint, Error, TEXT("Failed to convert Convert node from '%s' to '%s'"), *OpNode->OperationName.ToString(), *NewOpName.ToString());
+	}
+}
+
 FText UK2Node_PromotableOperator::GetTooltipText() const
 {
 	FFormatNamedArguments Args;
