Fix incorrect LOCTEXT namespace keys

#pf skip
#rnx

diff --git a/Engine/Source/Editor/BlueprintGraph/Private/K2Node_PromotableOperator.cpp b/Engine/Source/Editor/BlueprintGraph/Private/K2Node_PromotableOperator.cpp
--- a/Engine/Source/Editor/BlueprintGraph/Private/K2Node_PromotableOperator.cpp
+++ b/Engine/Source/Editor/BlueprintGraph/Private/K2Node_PromotableOperator.cpp
@@ -163,12 +163,12 @@ void UK2Node_PromotableOperator::GetNodeContextMenuActions(UToolMenu* Menu, UGra
 			Args.Add(TEXT("NewOpName"), FTypePromotion::GetUserFacingOperatorName(PossibleConversionOpName));
 			Args.Add(TEXT("CurrentOpName"), FTypePromotion::GetUserFacingOperatorName(OperationName));
 
-			const FText PinConversionName = FText::Format(LOCTEXT("CallFunction_Tooltip", "Convert to {NewOpName}"), Args);
+			const FText PinConversionName = FText::Format(LOCTEXT("ConvertOpName_Tooltip", "Convert to {NewOpName}"), Args);
 
 			Section.AddMenuEntry(
 				FName(PinConversionName.ToString()),
 				PinConversionName,
-				FText::Format(LOCTEXT("ConvertPinTypeTooltip", "Convert this node operation from '{CurrentOpName}' to '{NewOpName}'"), Args),
+				FText::Format(LOCTEXT("ConvertOperator_ToType_Tooltip", "Convert this node operation from '{CurrentOpName}' to '{NewOpName}'"), Args),
 				FSlateIcon(),
 				FUIAction(
 					FExecuteAction::CreateStatic(&UK2Node_PromotableOperator::ConvertComparisonOperatorNode, const_cast<UEdGraphNode*>(Context->Node.Get()), PossibleConversionOpName)
@@ -282,7 +282,7 @@ void UK2Node_PromotableOperator::CreateConversionMenu(FToolMenuSection& Conversi
 		Args.Add(TEXT("NewPinType"), Schema->TypeToText(PinType));
 		Args.Add(TEXT("CurrentPinType"), Schema->TypeToText(OriginalContextType));
 
-		const FText PinConversionName = FText::Format(LOCTEXT("CallFunction_Tooltip", "To {NewPinType}"), Args);
+		const FText PinConversionName = FText::Format(LOCTEXT("Convert_Pin_To_Type_Tooltip", "To {NewPinType}"), Args);
 
 		ConversionSection.AddMenuEntry(
 			FName(PinConversionName.ToString()),
