Allow you to convert operator nodes even if there are no default values or pin connections to the node

#rnx

diff --git a/Engine/Source/Editor/BlueprintGraph/Private/K2Node_PromotableOperator.cpp b/Engine/Source/Editor/BlueprintGraph/Private/K2Node_PromotableOperator.cpp
--- a/Engine/Source/Editor/BlueprintGraph/Private/K2Node_PromotableOperator.cpp
+++ b/Engine/Source/Editor/BlueprintGraph/Private/K2Node_PromotableOperator.cpp
@@ -353,6 +353,18 @@ void UK2Node_PromotableOperator::ConvertComparisonOperatorNode(UEdGraphNode* Nod
 	TArray<UEdGraphPin*> PinsToConsider;
 	OpNode->GetPinsToConsider(PinsToConsider);
 
+	// In the case of this node having no pins to consider (it is a wildcard, with no default values or connections)
+	// we will just use the boolean output pin to ensure that we get a valid UFunction to use.
+	if (PinsToConsider.IsEmpty())
+	{
+		UEdGraphPin* OutPin = OpNode->GetOutputPin();
+		// Because we only allow this conversion to happen on comparison operators, we can be sure the output pin will be a 
+		// simple boolean output and that it is there
+		ensure(OutPin && OutPin->PinType.PinCategory == UEdGraphSchema_K2::PC_Boolean);
+		PinsToConsider.Add(OutPin);		
+	}
+
+	// For nodes with connections we have to find the best function again that matches for them
 	if (const UFunction* BestMatchingFunc = FTypePromotion::FindBestMatchingFunc(OpNode->OperationName, PinsToConsider))
 	{
 		// Only allow this with comparison functions
