Change the blueprint action database to refresh type promotion instead of doing it from the module change callback.
This ensures the type promotion is refreshed right before it is needed and avoids redundant refreshes when multiple modules load at once.

diff --git a/Engine/Source/Editor/BlueprintGraph/Classes/BlueprintTypePromotion.h b/Engine/Source/Editor/BlueprintGraph/Classes/BlueprintTypePromotion.h
--- a/Engine/Source/Editor/BlueprintGraph/Classes/BlueprintTypePromotion.h
+++ b/Engine/Source/Editor/BlueprintGraph/Classes/BlueprintTypePromotion.h
@@ -9,7 +9,6 @@
 #include "Containers/Map.h"
 #include "Containers/Set.h"
 #include "CoreMinimal.h"
-#include "Delegates/IDelegateInstance.h"
 #include "EdGraph/EdGraphPin.h"
 #include "EdGraphSchema_K2.h"
 #include "HAL/IConsoleManager.h"
@@ -152,9 +151,6 @@ private:
 	explicit FTypePromotion();
 	~FTypePromotion();
 
-	/** Callback that will rebuild the op table when hot reload is triggered */
-	static void OnModulesChanged(FName ModuleThatChanged, EModuleChangeReason ReasonForChange);
-
 	bool IsFunctionPromotionReady_Internal(const UFunction* const FuncToConsider) const;
 
 	FEdGraphPinType GetPromotedType_Internal(const TArray<UEdGraphPin*>& WildcardPins) const;
@@ -183,12 +179,6 @@ private:
 
 	static FTypePromotion* Instance;
 
-	/** Delegate to handle that will be used to refresh the op table when a module has changed */
-	FDelegateHandle OnModulesChangedDelegateHandle;
-
-	/** Handle that will be used to refresh the top table when a Live coding reload has completed */
-	FDelegateHandle OnReloadCompleteDelegateHandle;
-
 	/** A map of 'Type' to its 'available promotions'. See ctor for creation */
 	const TMap<FName, TArray<FName>> PromotionTable;
 
diff --git a/Engine/Source/Editor/BlueprintGraph/Private/BlueprintActionDatabase.cpp b/Engine/Source/Editor/BlueprintGraph/Private/BlueprintActionDatabase.cpp
--- a/Engine/Source/Editor/BlueprintGraph/Private/BlueprintActionDatabase.cpp
+++ b/Engine/Source/Editor/BlueprintGraph/Private/BlueprintActionDatabase.cpp
@@ -1365,6 +1365,10 @@ void FBlueprintActionDatabase::Tick(float DeltaTime)
 	{
 		RefreshAll();
 	}
+	else if (!BlueprintActionDatabaseImpl::PendingModules.IsEmpty())
+	{
+		PreRefresh(false);
+	}
 	
 	// Check for any modules that may have been loaded since the last tick. Even if we call RefreshAll() above, we still want to run
 	// through this list in order to keep track of loaded modules containing native script types that are registered into the database.
@@ -1500,6 +1504,9 @@ void FBlueprintActionDatabase::RefreshAll()
 	TGuardValue<bool> ScopedInitialization(BlueprintActionDatabaseImpl::bIsInitializing, true);
 	BlueprintActionDatabaseImpl::bRefreshAllRequested = false;
 
+	// Refresh other systems before the database is recreated
+	PreRefresh(true);
+
 	// Remove callbacks from blueprints
 	for (TObjectIterator<UBlueprint> BlueprintIt; BlueprintIt; ++BlueprintIt)
 	{
@@ -2002,6 +2009,12 @@ void FBlueprintActionDatabase::OnBlueprintChanged(UBlueprint* InBlueprint)
 	}
 }
 
+void FBlueprintActionDatabase::PreRefresh(bool bRefreshAll)
+{
+	// Refresh other systems as necessary, doing it here avoids redundant work
+	FTypePromotion::RefreshPromotionTables();
+}
+
 bool FBlueprintActionDatabase::IsClassAllowed(UClass const* InClass, EPermissionsContext InContext)
 {
 	if (InClass == nullptr)
diff --git a/Engine/Source/Editor/BlueprintGraph/Private/BlueprintEditorSettings.cpp b/Engine/Source/Editor/BlueprintGraph/Private/BlueprintEditorSettings.cpp
--- a/Engine/Source/Editor/BlueprintGraph/Private/BlueprintEditorSettings.cpp
+++ b/Engine/Source/Editor/BlueprintGraph/Private/BlueprintEditorSettings.cpp
@@ -9,7 +9,6 @@
 #include "BlueprintEditorModule.h"
 #include "BlueprintNamespaceHelper.h"
 #include "BlueprintNamespaceUtilities.h"
-#include "BlueprintTypePromotion.h"
 #include "CoreGlobals.h"
 #include "EdGraphSchema_K2.h"
 #include "Editor/EditorPerProjectUserSettings.h"
@@ -148,12 +147,10 @@ void UBlueprintEditorSettings::PostEditChangeProperty(FPropertyChangedEvent& Pro
 		bShouldRebuildRegistry = true;
 	}
 	
-	// Refresh type promotion when the preference gets changed so that we can correctly rebuild the action database
+	// Type promotion changes are handled by the action database refresh
 	if (PropertyName == GET_MEMBER_NAME_CHECKED(UBlueprintEditorSettings, bEnableTypePromotion) || 
 		PropertyName == GET_MEMBER_NAME_CHECKED(UBlueprintEditorSettings, TypePromotionPinDenyList))
 	{
-		FTypePromotion::RefreshPromotionTables();
-		
 		bShouldRebuildRegistry = true;
 	}
 
diff --git a/Engine/Source/Editor/BlueprintGraph/Private/BlueprintTypePromotion.cpp b/Engine/Source/Editor/BlueprintGraph/Private/BlueprintTypePromotion.cpp
--- a/Engine/Source/Editor/BlueprintGraph/Private/BlueprintTypePromotion.cpp
+++ b/Engine/Source/Editor/BlueprintGraph/Private/BlueprintTypePromotion.cpp
@@ -5,7 +5,6 @@
 #include "Async/ParallelFor.h"
 #include "Containers/EnumAsByte.h"
 #include "Containers/UnrealString.h"
-#include "Delegates/Delegate.h"
 #include "EdGraph/EdGraphNode.h"
 #include "EdGraphSchema_K2.h"
 #include "Internationalization/Internationalization.h"
@@ -18,7 +17,6 @@
 #include "UObject/Class.h"
 #include "UObject/ObjectMacros.h"
 #include "UObject/Script.h"
-#include "UObject/UObjectGlobals.h"		// For FCoreUObjectDelegates::ReloadCompleteDelegate
 #include "UObject/UObjectHash.h"
 #include "UObject/UnrealNames.h"
 #include "UObject/UnrealType.h"
@@ -77,25 +75,10 @@ FTypePromotion::FTypePromotion()
 	: PromotionTable(CreatePromotionTable())
 {
 	CreateOpTable();
-	OnModulesChangedDelegateHandle = FModuleManager::Get().OnModulesChanged().AddStatic(&FTypePromotion::OnModulesChanged);
-	OnReloadCompleteDelegateHandle = FCoreUObjectDelegates::ReloadCompleteDelegate.AddStatic(&FTypePromotion::RefreshPromotionTables);
 }
 
 FTypePromotion::~FTypePromotion()
 {
-	FModuleManager::Get().OnModulesChanged().Remove(OnModulesChangedDelegateHandle);
-	FCoreUObjectDelegates::ReloadCompleteDelegate.Remove(OnReloadCompleteDelegateHandle);
-}
-
-void FTypePromotion::OnModulesChanged(FName ModuleThatChanged, EModuleChangeReason ReasonForChange)
-{
-	// Any time a module is changed, there could possibly be new UFunctions that we 
-	// need to process, so we need to recreate the op table and clear the node spawners
-	// that we are using in order to avoid invalid duplicates in the graph action menu
-	if (ReasonForChange == EModuleChangeReason::ModuleLoaded || ReasonForChange == EModuleChangeReason::ModuleUnloaded)
-	{
-		FTypePromotion::RefreshPromotionTables();
-	}
 }
 
 TMap<FName, TArray<FName>> FTypePromotion::CreatePromotionTable()
diff --git a/Engine/Source/Editor/BlueprintGraph/Public/BlueprintActionDatabase.h b/Engine/Source/Editor/BlueprintGraph/Public/BlueprintActionDatabase.h
--- a/Engine/Source/Editor/BlueprintGraph/Public/BlueprintActionDatabase.h
+++ b/Engine/Source/Editor/BlueprintGraph/Public/BlueprintActionDatabase.h
@@ -242,6 +242,9 @@ private:
 	 * This exists only because we need a pointer to associate our delegates with
 	 */
 	void OnBlueprintChanged( UBlueprint* );
+
+	/** Refresh other systems before a full or partial refresh */
+	void PreRefresh(bool bRefreshAll);
 private:
 	/** 
 	 * A map of associated node-spawners for each class/asset. A spawner that 
