Fix tooltip being incorrect after an undo of converting promotable operator type nodes. The operation name was not a UPROPERTY, so it was not being tracked with Modify calls for transactions

#rnx

diff --git a/Engine/Source/Editor/BlueprintGraph/Classes/K2Node_PromotableOperator.h b/Engine/Source/Editor/BlueprintGraph/Classes/K2Node_PromotableOperator.h
--- a/Engine/Source/Editor/BlueprintGraph/Classes/K2Node_PromotableOperator.h
+++ b/Engine/Source/Editor/BlueprintGraph/Classes/K2Node_PromotableOperator.h
@@ -164,6 +164,7 @@ private:
 	void CreateConversionMenu(struct FToolMenuSection& ConversionSection, UEdGraphPin* ContextPin) const;
 
 	/** The name that this operation uses ("Add", "Multiply", etc) */
+	UPROPERTY()
 	FName OperationName;
 
 	/** The current number of additional pins on this node */
