Fix for an editor crash after editing Math Expression node titles while hovered over the content area.

diff --git a/Engine/Source/Editor/BlueprintGraph/Private/K2Node_MathExpression.cpp b/Engine/Source/Editor/BlueprintGraph/Private/K2Node_MathExpression.cpp
--- a/Engine/Source/Editor/BlueprintGraph/Private/K2Node_MathExpression.cpp
+++ b/Engine/Source/Editor/BlueprintGraph/Private/K2Node_MathExpression.cpp
@@ -18,6 +18,7 @@
 #include "EdGraphUtilities.h"
 #include "Engine/Blueprint.h"
 #include "Engine/MemberReference.h"
+#include "Framework/Application/SlateApplication.h"
 #include "HAL/PlatformCrt.h"
 #include "Internationalization/Internationalization.h"
 #include "K2Node_CallFunction.h"
@@ -2731,6 +2732,10 @@ void UK2Node_MathExpression::RebuildExpression(FString InExpression)
 //------------------------------------------------------------------------------
 void UK2Node_MathExpression::ClearExpression()
 {
+	// close the tooltip preview (if open) - since this is a composite node,
+	// its content will become invalid because we are mutating the bound graph
+	FSlateApplication::Get().CloseToolTip();
+
 	// clear any errors 
 	SetNodeError(this, FText::GetEmpty());
 
diff --git a/Engine/Source/Editor/GraphEditor/Private/ConnectionDrawingPolicy.cpp b/Engine/Source/Editor/GraphEditor/Private/ConnectionDrawingPolicy.cpp
--- a/Engine/Source/Editor/GraphEditor/Private/ConnectionDrawingPolicy.cpp
+++ b/Engine/Source/Editor/GraphEditor/Private/ConnectionDrawingPolicy.cpp
@@ -475,6 +475,11 @@ void FConnectionDrawingPolicy::DrawPinGeometries(TMap<TSharedRef<SWidget>, FArra
 		SGraphPin& PinWidget = static_cast<SGraphPin&>(SomePinWidget.Get());
 		UEdGraphPin* ThePin = PinWidget.GetPinObj();
 
+		if (!ensureMsgf(ThePin, TEXT("Encountered a NULL graph pin object while attempting to draw connections. It may be necessary to refresh the outer node widget.")))
+		{
+			continue;
+		}
+
 		if (ThePin->Direction == EGPD_Output)
 		{
 			for (int32 LinkIndex=0; LinkIndex < ThePin->LinkedTo.Num(); ++LinkIndex)
