CIS fix (373586).

#rnx

diff --git a/Engine/Source/Editor/BlueprintGraph/Private/K2Node_MathExpression.cpp b/Engine/Source/Editor/BlueprintGraph/Private/K2Node_MathExpression.cpp
--- a/Engine/Source/Editor/BlueprintGraph/Private/K2Node_MathExpression.cpp
+++ b/Engine/Source/Editor/BlueprintGraph/Private/K2Node_MathExpression.cpp
@@ -2732,9 +2732,13 @@ void UK2Node_MathExpression::RebuildExpression(FString InExpression)
 //------------------------------------------------------------------------------
 void UK2Node_MathExpression::ClearExpression()
 {
-	// close the tooltip preview (if open) - since this is a composite node,
-	// its content will become invalid because we are mutating the bound graph
-	FSlateApplication::Get().CloseToolTip();
+	if (FSlateApplication::IsInitialized())
+	{
+		// close the tooltip preview (if open) - since this is a composite node,
+		// its content will become invalid because we are mutating the bound graph
+		// @todo - see if we can find a better place to do this on the Slate/UI side
+		FSlateApplication::Get().CloseToolTip();
+	}
 
 	// clear any errors 
 	SetNodeError(this, FText::GetEmpty());
