Fixed usage of CpuProfiler trace in BlueprintActionFilter (when ENABLE_BLUEPRINT_ACTION_FILTER_PROFILING is enabled). FCpuProfilerTrace::OutputEventType and FCpuProfilerTrace::OutputBeginDynamicEvent needs to be called only if CpuChannel is enabled, otherwise the "spec" trace event is not emitted, so it will result in <unknown> cpu timers. Also fixed code that matches OutputBeginEvent with OutputEndEvent.

diff --git a/Engine/Source/Editor/BlueprintGraph/Private/BlueprintActionFilter.cpp b/Engine/Source/Editor/BlueprintGraph/Private/BlueprintActionFilter.cpp
--- a/Engine/Source/Editor/BlueprintGraph/Private/BlueprintActionFilter.cpp
+++ b/Engine/Source/Editor/BlueprintGraph/Private/BlueprintActionFilter.cpp
@@ -501,6 +501,7 @@ namespace BlueprintActionFilterImpl
 		uint32* CpuSpecIdPtr = FuncPtrToCpuSpecIdMap.Find(TestFuncPtr);
 		if (!CpuSpecIdPtr)
 		{
+			check(UE_TRACE_CHANNELEXPR_IS_ENABLED(CpuChannel));
 			CpuSpecIdPtr = &FuncPtrToCpuSpecIdMap.Add(TestFuncPtr, FCpuProfilerTrace::OutputEventType(*FilterTestFuncPtrToFuncName(TestFuncPtr)));
 		}
 
@@ -2357,12 +2358,16 @@ bool FBlueprintActionFilter::IsFilteredByThis(FBlueprintActionInfo& BlueprintAct
 		checkSlow(RejectionTestDelegate.IsBound());
 
 #if ENABLE_BLUEPRINT_ACTION_FILTER_PROFILING
+		bool bTraceEventEmitted = false;
 		double StartTime = 0;
 		FFilterTestProfileRecord* FilterTestProfileRecord = nullptr;
 		if (bIsFilterTestProfilingEnabled)
 		{
-			FCpuProfilerTrace::OutputBeginDynamicEvent(BlueprintActionFilterImpl::FilterTestProfileEventName);
-
+			if (UE_TRACE_CHANNELEXPR_IS_ENABLED(CpuChannel))
+			{
+				FCpuProfilerTrace::OutputBeginDynamicEvent(BlueprintActionFilterImpl::FilterTestProfileEventName);
+				bTraceEventEmitted = true;
+			}
 			StartTime = FPlatformTime::Seconds();
 			FilterTestProfileRecord = BeginFilterTestProfileEvent(TestIndex, RejectionTestDelegate.GetBoundProgramCounterForTimerManager());
 		}
@@ -2374,7 +2379,10 @@ bool FBlueprintActionFilter::IsFilteredByThis(FBlueprintActionInfo& BlueprintAct
 		if (FilterTestProfileRecord)
 		{
 			EndFilterTestProfileEvent(FilterTestProfileRecord, bIsFiltered, StartTime);
-			FCpuProfilerTrace::OutputEndEvent();
+			if (bTraceEventEmitted)
+			{
+				FCpuProfilerTrace::OutputEndEvent();
+			}
 		}
 #endif	// ENABLE_BLUEPRINT_ACTION_FILTER_PROFILING
 
@@ -2393,12 +2401,16 @@ bool FBlueprintActionFilter::IsFilteredByThis(FBlueprintActionInfo& BlueprintAct
 		for (const FBlueprintGraphModule::FActionMenuRejectionTest& ExtraRejectionTest : BluprintGraphModule->GetExtendedActionMenuFilters())
 		{
 #if ENABLE_BLUEPRINT_ACTION_FILTER_PROFILING
+			bool bTraceEventEmitted = false;
 			double StartTime = 0;
 			FFilterTestProfileRecord* FilterTestProfileRecord = nullptr;
 			if (bIsFilterTestProfilingEnabled)
 			{
-				FCpuProfilerTrace::OutputBeginDynamicEvent(BlueprintActionFilterImpl::FilterTestProfileEventName);
-
+				if (UE_TRACE_CHANNELEXPR_IS_ENABLED(CpuChannel))
+				{
+					FCpuProfilerTrace::OutputBeginDynamicEvent(BlueprintActionFilterImpl::FilterTestProfileEventName);
+					bTraceEventEmitted = true;
+				}
 				StartTime = FPlatformTime::Seconds();
 				FilterTestProfileRecord = BeginFilterTestProfileEvent(ExtraTestProfileRecordIndex++, ExtraRejectionTest.GetBoundProgramCounterForTimerManager());
 			}
@@ -2410,7 +2422,10 @@ bool FBlueprintActionFilter::IsFilteredByThis(FBlueprintActionInfo& BlueprintAct
 			if (FilterTestProfileRecord)
 			{
 				EndFilterTestProfileEvent(FilterTestProfileRecord, bIsFiltered, StartTime);
-				FCpuProfilerTrace::OutputEndEvent();
+				if (bTraceEventEmitted)
+				{
+					FCpuProfilerTrace::OutputEndEvent();
+				}
 			}
 #endif	// ENABLE_BLUEPRINT_ACTION_FILTER_PROFILING
 
@@ -2475,9 +2490,11 @@ FBlueprintActionFilter::FFilterTestProfileRecord* FBlueprintActionFilter::BeginF
 		FilterTestProfileRecord->TestFuncPtr = FilterTestFuncPtr;
 	}
 
-	if (IsFilterTestTraceLoggingEnabled())
+	if (IsFilterTestTraceLoggingEnabled() && UE_TRACE_CHANNELEXPR_IS_ENABLED(CpuChannel))
 	{
 		FCpuProfilerTrace::OutputBeginEvent(BlueprintActionFilterImpl::FilterTestFuncPtrToCpuSpecId(FilterTestFuncPtr));
+		check(FilterTestProfileRecord->bOutputBeginEventEmitted == false);
+		FilterTestProfileRecord->bOutputBeginEventEmitted = true;
 	}
 
 	return FilterTestProfileRecord;
@@ -2487,9 +2504,10 @@ void FBlueprintActionFilter::EndFilterTestProfileEvent(FFilterTestProfileRecord*
 {
 	if (InEvent)
 	{
-		if (IsFilterTestTraceLoggingEnabled())
+		if (InEvent->bOutputBeginEventEmitted)
 		{
 			FCpuProfilerTrace::OutputEndEvent();
+			InEvent->bOutputBeginEventEmitted = false;
 		}
 
 		if (IsFilterTestStatsLoggingEnabled())
diff --git a/Engine/Source/Editor/BlueprintGraph/Public/BlueprintActionFilter.h b/Engine/Source/Editor/BlueprintGraph/Public/BlueprintActionFilter.h
--- a/Engine/Source/Editor/BlueprintGraph/Public/BlueprintActionFilter.h
+++ b/Engine/Source/Editor/BlueprintGraph/Public/BlueprintActionFilter.h
@@ -444,6 +444,7 @@ private:
 		float TotalTimeMs = 0.0f;
 		int32 NumIterations = 0;
 		int32 NumFilteredOut = 0;
+		bool bOutputBeginEventEmitted = false;
 	};
 
 	mutable TMap<int32, FFilterTestProfileRecord> FilterTestProfiles;
