Avoid propagating pin type change, when no pin has changed type

diff --git a/Engine/Source/Editor/BlueprintGraph/Private/K2Node_Knot.cpp b/Engine/Source/Editor/BlueprintGraph/Private/K2Node_Knot.cpp
--- a/Engine/Source/Editor/BlueprintGraph/Private/K2Node_Knot.cpp
+++ b/Engine/Source/Editor/BlueprintGraph/Private/K2Node_Knot.cpp
@@ -167,6 +167,14 @@ void UK2Node_Knot::PropagatePinTypeFromDirection(bool bFromInput)
 	UEdGraphPin* TypeSource = MySourcePin->LinkedTo.Num() ? MySourcePin->LinkedTo[0] : nullptr;
 	if (TypeSource)
 	{
+		// if the type in the source and dest matches the type source then
+		// lets early return to avoid expensive propagation in PinConnectionListChanged
+		if (MySourcePin->PinType == TypeSource->PinType &&
+			MyDestinationPin->PinType == TypeSource->PinType)
+		{
+			return;
+		}
+
 		MySourcePin->PinType = TypeSource->PinType;
 		MyDestinationPin->PinType = TypeSource->PinType;
 
