PR #9602 Property meta specifier DeterminesOutputType now works with Soft Object and Soft Class Paths

diff --git a/Engine/Source/Editor/BlueprintGraph/Private/K2Node_CallFunction.cpp b/Engine/Source/Editor/BlueprintGraph/Private/K2Node_CallFunction.cpp
--- a/Engine/Source/Editor/BlueprintGraph/Private/K2Node_CallFunction.cpp
+++ b/Engine/Source/Editor/BlueprintGraph/Private/K2Node_CallFunction.cpp
@@ -49,6 +49,7 @@
 #include "BlueprintNodeStatics.h"
 #include "Settings/BlueprintEditorProjectSettings.h"
 #include "ToolMenu.h"
+#include "FindInBlueprints.h"
 
 #define LOCTEXT_NAMESPACE "K2Node"
 
@@ -302,7 +303,12 @@ UClass* FDynamicOutputHelper::GetPinClass(UEdGraphPin* Pin)
 {
 	UClass* PinClass = UObject::StaticClass();
 
-	bool const bIsClassOrObjectPin = (Pin->PinType.PinCategory == UEdGraphSchema_K2::PC_Class || Pin->PinType.PinCategory == UEdGraphSchema_K2::PC_Object);
+	bool const bIsClassOrObjectPin = 
+		   Pin->PinType.PinCategory == UEdGraphSchema_K2::PC_Class 
+		|| Pin->PinType.PinCategory == UEdGraphSchema_K2::PC_Object 
+		|| Pin->PinType.PinCategory == UEdGraphSchema_K2::PC_SoftClass
+		|| Pin->PinType.PinCategory == UEdGraphSchema_K2::PC_SoftObject;
+
 	if (bIsClassOrObjectPin)
 	{
 		if (UClass* DefaultClass = Cast<UClass>(Pin->DefaultObject))
@@ -314,6 +320,25 @@ UClass* FDynamicOutputHelper::GetPinClass(UEdGraphPin* Pin)
 			PinClass = BaseClass;
 		}
 
+		// If the pin's default value is a soft class or object path, resolve the class/object and use that as type
+		if (Pin->PinType.PinCategory == UEdGraphSchema_K2::PC_SoftClass)
+		{
+			const FSoftClassPath SoftClassPath(Pin->DefaultValue);
+			if (UClass* DefaultValueClass = SoftClassPath.TryLoadClass<UObject>())
+			{
+				PinClass = DefaultValueClass;
+			}
+		}
+		else if (Pin->PinType.PinCategory == UEdGraphSchema_K2::PC_SoftObject)
+		{
+			const FSoftObjectPtr SoftObjPtr(FSoftObjectPath(Pin->DefaultValue));
+			if (UObject* DefaultValueObject = SoftObjPtr.LoadSynchronous())
+			{
+				PinClass = DefaultValueObject->GetClass();
+			}
+		}
+
+		// If the pin has connections, derive the pin's type from the set of connected pins
 		if (Pin->LinkedTo.Num() > 0)
 		{
 			UClass* CommonInputClass = nullptr;
@@ -334,6 +359,7 @@ UClass* FDynamicOutputHelper::GetPinClass(UEdGraphPin* Pin)
 				{
 					if (CommonInputClass != nullptr)
 					{
+						// Update common super class of all linked pins
 						while (!LinkClass->IsChildOf(CommonInputClass))
 						{
 							CommonInputClass = CommonInputClass->GetSuperClass();
@@ -1589,6 +1615,24 @@ FLinearColor UK2Node_CallFunction::GetNodeTitleColor() const
 	return GetPalletteIconColor(GetTargetFunction());
 }
 
+FString UK2Node_CallFunction::GetFindReferenceSearchString(const bool bUseSearchSyntax) const
+{
+	if (bUseSearchSyntax)
+	{
+		if (const UFunction* Function = GetTargetFunction())
+		{
+			FString SearchTerm;
+			if (FindInBlueprintsHelpers::ConstructSearchTermFromFunction(Function, SearchTerm))
+			{
+				return SearchTerm;
+			}
+		}
+	}
+
+	// Fallback behavior
+	return Super::GetFindReferenceSearchString(bUseSearchSyntax);
+}
+
 FText UK2Node_CallFunction::GetTooltipText() const
 {
 	FText Tooltip;
@@ -3510,6 +3554,28 @@ void UK2Node_CallFunction::AddSearchMetaDataInfo(TArray<struct FSearchTagDataPai
 	}
 }
 
+void UK2Node_CallFunction::AddPinSearchMetaDataInfo(const UEdGraphPin* Pin, TArray<FSearchTagDataPair>& OutTaggedMetaData) const
+{
+	Super::AddPinSearchMetaDataInfo(Pin, OutTaggedMetaData);
+
+	// Blueprint graphs that call a function declared in the same blueprint don't store a target type, but rather PinSubCategory == Self.
+	// When this is the case, we will still explicitly index the ObjectClass for the target pin, so that it can be treated the same as 
+	// any other call function nodes.
+	if (Pin->PinName == UEdGraphSchema_K2::PN_Self && Pin->PinType.PinSubCategory == UEdGraphSchema_K2::PSC_Self && !Pin->PinType.PinSubCategoryObject.IsValid())
+	{
+		if (const UFunction* Function = GetTargetFunction())
+		{
+			if (Function->GetOwnerClass() && Function->GetOwnerClass()->GetAuthoritativeClass())
+			{
+				// Index the name of the class that declares the function, this will be used in search queries.
+				// When PinSubCategory == self, we don't expect that the function is an override, so no need to consider super functions.
+				const FString ObjectClass = Function->GetOwnerClass()->GetAuthoritativeClass()->GetName();
+				OutTaggedMetaData.Add(FSearchTagDataPair(FFindInBlueprintSearchTags::FiB_ObjectClass, FText::FromString(*ObjectClass)));
+			}
+		}
+	}
+}
+
 TSharedPtr<SWidget> UK2Node_CallFunction::CreateNodeImage() const
 {
 	// For set, map and array functions we have a cool icon. This helps users quickly
