PR #9602 Property meta specifier DeterminesOutputType now works with Soft Object and Soft Class Paths

diff --git a/Engine/Source/Editor/BlueprintGraph/Private/K2Node_CallFunction.cpp b/Engine/Source/Editor/BlueprintGraph/Private/K2Node_CallFunction.cpp
--- a/Engine/Source/Editor/BlueprintGraph/Private/K2Node_CallFunction.cpp
+++ b/Engine/Source/Editor/BlueprintGraph/Private/K2Node_CallFunction.cpp
@@ -302,7 +302,12 @@ UClass* FDynamicOutputHelper::GetPinClass(UEdGraphPin* Pin)
 {
 	UClass* PinClass = UObject::StaticClass();
 
-	bool const bIsClassOrObjectPin = (Pin->PinType.PinCategory == UEdGraphSchema_K2::PC_Class || Pin->PinType.PinCategory == UEdGraphSchema_K2::PC_Object);
+	bool const bIsClassOrObjectPin = 
+		   Pin->PinType.PinCategory == UEdGraphSchema_K2::PC_Class 
+		|| Pin->PinType.PinCategory == UEdGraphSchema_K2::PC_Object 
+		|| Pin->PinType.PinCategory == UEdGraphSchema_K2::PC_SoftClass
+		|| Pin->PinType.PinCategory == UEdGraphSchema_K2::PC_SoftObject;
+
 	if (bIsClassOrObjectPin)
 	{
 		if (UClass* DefaultClass = Cast<UClass>(Pin->DefaultObject))
@@ -314,6 +319,25 @@ UClass* FDynamicOutputHelper::GetPinClass(UEdGraphPin* Pin)
 			PinClass = BaseClass;
 		}
 
+		// If the pin's default value is a soft class or object path, resolve the class/object and use that as type
+		if (Pin->PinType.PinCategory == UEdGraphSchema_K2::PC_SoftClass)
+		{
+			const FSoftClassPath SoftClassPath(Pin->DefaultValue);
+			if (UClass* DefaultValueClass = SoftClassPath.TryLoadClass<UObject>())
+			{
+				PinClass = DefaultValueClass;
+			}
+		}
+		else if (Pin->PinType.PinCategory == UEdGraphSchema_K2::PC_SoftObject)
+		{
+			const FSoftObjectPtr SoftObjPtr(FSoftObjectPath(Pin->DefaultValue));
+			if (UObject* DefaultValueObject = SoftObjPtr.LoadSynchronous())
+			{
+				PinClass = DefaultValueObject->GetClass();
+			}
+		}
+
+		// If the pin has connections, derive the pin's type from the set of connected pins
 		if (Pin->LinkedTo.Num() > 0)
 		{
 			UClass* CommonInputClass = nullptr;
@@ -334,6 +358,7 @@ UClass* FDynamicOutputHelper::GetPinClass(UEdGraphPin* Pin)
 				{
 					if (CommonInputClass != nullptr)
 					{
+						// Update common super class of all linked pins
 						while (!LinkClass->IsChildOf(CommonInputClass))
 						{
 							CommonInputClass = CommonInputClass->GetSuperClass();
