Removed check in ValidateLinkedPinTypes

Some users are hitting a check where no compiler relevant pin cannot be found during validation. It's unclear if this is an expected scenario. For now, we'll relax this to a soft-fail, and add logging to help track down the issue. If this turns out to be a valid scenario, we can remove the log in a future change.

diff --git a/Engine/Source/Editor/BlueprintGraph/Private/K2Node.cpp b/Engine/Source/Editor/BlueprintGraph/Private/K2Node.cpp
--- a/Engine/Source/Editor/BlueprintGraph/Private/K2Node.cpp
+++ b/Engine/Source/Editor/BlueprintGraph/Private/K2Node.cpp
@@ -1164,9 +1164,16 @@ void UK2Node::ValidateLinkedPinTypes(UEdGraphPin* OutputPin, FCompilerResultsLog
 					Args
 				);
 
-				const UEdGraphPin* SourcePin = FBlueprintEditorUtils::FindFirstCompilerRelevantLinkedPin(OutputPin);
-				check(SourcePin);
-				Message_Note(ConversionInfo.ToString(), SourcePin->GetOwningNode(), InputPin->GetOwningNode());
+				if (const UEdGraphPin* SourcePin = FBlueprintEditorUtils::FindFirstCompilerRelevantLinkedPin(OutputPin))
+				{
+					Message_Note(ConversionInfo.ToString(), SourcePin->GetOwningNode(), InputPin->GetOwningNode());
+				}
+				else
+				{
+					const UK2Node* OutputPinOwner = Cast<UK2Node>(OutputPin->GetOwningNode());
+
+					UE_LOG(LogBlueprint, Verbose, TEXT("Missing compiler relevant pin '%s' on node '%s'"), *OutputPin->GetName(), (OutputPinOwner ? *OutputPinOwner->GetFullName() : TEXT("<none>")));
+				}
 			}
 		}
 	}
