Fix editor crash after changing the title content on a math expression graph node.

diff --git a/Engine/Source/Editor/BlueprintGraph/Private/K2Node_MathExpression.cpp b/Engine/Source/Editor/BlueprintGraph/Private/K2Node_MathExpression.cpp
--- a/Engine/Source/Editor/BlueprintGraph/Private/K2Node_MathExpression.cpp
+++ b/Engine/Source/Editor/BlueprintGraph/Private/K2Node_MathExpression.cpp
@@ -1963,8 +1963,9 @@ private:
 		GraphXBounds.X = FMath::Min((int32)GraphXBounds.X, X);
 		GraphXBounds.Y = FMath::Max((int32)GraphXBounds.Y, X);
 
+		constexpr bool bShouldSelectNewNode = false;
 		const FVector2D Location = GetNodePosition(X, Y);
-		return FEdGraphSchemaAction_K2NewNode::SpawnNodeFromTemplate<NodeType>(CompilingNode->BoundGraph, Template, Location);
+		return FEdGraphSchemaAction_K2NewNode::SpawnNodeFromTemplate<NodeType>(CompilingNode->BoundGraph, Template, Location, bShouldSelectNewNode);
 	}
 
 private:
@@ -2719,6 +2720,12 @@ void UK2Node_MathExpression::RebuildExpression(FString InExpression)
 			}
 		}
 
+		// notify any listeners that the bound graph has changed
+		if (BoundGraph)
+		{
+			BoundGraph->NotifyGraphChanged();
+		}
+
 		// refresh the node since the connections may have changed, this won't be reentrant due to bool above
 		ReconstructNode();
 
