Finding References By Name on event nodes (K2Node_Event) and function call nodes (K2Node_CallFunction) now consistently searches by native name, rather than friendly name in some cases. Additionally, K2Node_CreateDelegate appends metadata so it can be found when searching for the referenced function if it's a class member.

diff --git a/Engine/Source/Editor/BlueprintGraph/Private/K2Node_CallFunction.cpp b/Engine/Source/Editor/BlueprintGraph/Private/K2Node_CallFunction.cpp
--- a/Engine/Source/Editor/BlueprintGraph/Private/K2Node_CallFunction.cpp
+++ b/Engine/Source/Editor/BlueprintGraph/Private/K2Node_CallFunction.cpp
@@ -1594,6 +1594,7 @@ FString UK2Node_CallFunction::GetFindReferenceSearchString_Impl(EGetFindReferenc
 {
 	if (EnumHasAnyFlags(InFlags, EGetFindReferenceSearchStringFlags::UseSearchSyntax))
 	{
+		// Searching by class member: try to resolve function and construct precise query
 		if (const UFunction* Function = GetTargetFunction())
 		{
 			FString SearchTerm;
@@ -1603,8 +1604,17 @@ FString UK2Node_CallFunction::GetFindReferenceSearchString_Impl(EGetFindReferenc
 			}
 		}
 	}
+	else
+	{
+		// Searching by function name: try to resolve function to return native name
+		if (const UFunction* Function = GetTargetFunction())
+		{
+			const FString NativeName = Function->GetName();
+			return NativeName;
+		}
+	}
 
-	// Fallback behavior
+	// Fallback behavior: search by node title
 	return Super::GetFindReferenceSearchString_Impl(InFlags);
 }
 
diff --git a/Engine/Source/Editor/BlueprintGraph/Private/K2Node_CreateDelegate.cpp b/Engine/Source/Editor/BlueprintGraph/Private/K2Node_CreateDelegate.cpp
--- a/Engine/Source/Editor/BlueprintGraph/Private/K2Node_CreateDelegate.cpp
+++ b/Engine/Source/Editor/BlueprintGraph/Private/K2Node_CreateDelegate.cpp
@@ -466,6 +466,12 @@ void UK2Node_CreateDelegate::AddSearchMetaDataInfo(TArray<struct FSearchTagDataP
 	if (!FunctionName.IsNone())
 	{
 		OutTaggedMetaData.Add(FSearchTagDataPair(FFindInBlueprintSearchTags::FiB_NativeName, FText::FromName(FunctionName)));
+
+		if (const UClass* ScopeClass = GetScopeClass(/*bDontUseSkeletalClassForSelf=*/true))
+		{
+			const FString FuncOriginClassName = ScopeClass->GetPathName();
+			OutTaggedMetaData.Add(FSearchTagDataPair(FFindInBlueprintSearchTags::FiB_FuncOriginClass, FText::FromString(FuncOriginClassName)));
+		}
 	}
 }
 
diff --git a/Engine/Source/Editor/BlueprintGraph/Private/K2Node_Event.cpp b/Engine/Source/Editor/BlueprintGraph/Private/K2Node_Event.cpp
--- a/Engine/Source/Editor/BlueprintGraph/Private/K2Node_Event.cpp
+++ b/Engine/Source/Editor/BlueprintGraph/Private/K2Node_Event.cpp
@@ -901,25 +901,43 @@ FSlateIcon UK2Node_Event::GetIconAndTint(FLinearColor& OutColor) const
 
 FString UK2Node_Event::GetFindReferenceSearchString_Impl(EGetFindReferenceSearchStringFlags InFlags) const
 {
-	// Resolve the function
-	if (const UFunction* Function = FFunctionFromNodeHelper::FunctionFromNode(this))
+	// If searching by class member, try to construct search term from the UFunction.
+	// This may fail if the function was not found or for whatever reason, its owning 
+	// class is invalid. If it fails, proceed to search by name as fallback behavior.
+	if (EnumHasAnyFlags(InFlags, EGetFindReferenceSearchStringFlags::UseSearchSyntax))
+	{
+		// Resolve the function
+		if (const UFunction* Function = FFunctionFromNodeHelper::FunctionFromNode(this))
+		{
+			FString SearchTerm;
+			if (FindInBlueprintsHelpers::ConstructSearchTermFromFunction(Function, SearchTerm))
+			{
+				return SearchTerm;
+			}
+		}
+	}
+
+	// Searching by just name. When overriding a function, try to find the function that it overrides to return its name.
+	if (bOverrideFunction || (CustomFunctionName == NAME_None))
 	{
-		// Attempt to construct an advanced search syntax query from the function
-		FString SearchTerm;
-		if (EnumHasAnyFlags(InFlags, EGetFindReferenceSearchStringFlags::UseSearchSyntax) && FindInBlueprintsHelpers::ConstructSearchTermFromFunction(Function, SearchTerm))
+		// Attempt to find the function
+		if (const UFunction* Function = FFunctionFromNodeHelper::FunctionFromNode(this))
 		{
-			return SearchTerm;
+			// Search by native name
+			const FString FunctionNativeName = Function->GetName();
+			return FunctionNativeName;
 		}
 		else
 		{
-			// Fallback behavior: function was found but failed to construct a search term from it
-			// Just search for the function's friendly name
-			return UEdGraphSchema_K2::GetFriendlySignatureName(Function).ToString();
+			// If we fail to find the function, still want to search for its expected name
+			return EventReference.GetMemberName().ToString();
 		}
 	}
-	
-	// If we fail to find the function, still want to search for its expected name.
-	return EventReference.GetMemberName().ToString();
+	else
+	{
+		// The function was not an override; its name is defined by this node
+		return CustomFunctionName.ToString();
+	}
 }
 
 void UK2Node_Event::FindDiffs(UEdGraphNode* OtherNode, struct FDiffResults& Results)
