Allow for UObject Macro Libraries to use functions that require a world context.

If a UObject MacroLibrary goes to use a function that requires a world context then it will show the input pin for the world context just like it would if it were a blueprint function library.
The input pin for the WorldContext can be manually provided by the user by dragging the right reference and passing that into the Macro, but the Macro will also just auto grab the world from the caller if they have a valid world context. if they do not have a valid world context, then an error will be shown in the called at compile time.

diff --git a/Engine/Source/Editor/BlueprintGraph/Private/EdGraphSchema_K2.cpp b/Engine/Source/Editor/BlueprintGraph/Private/EdGraphSchema_K2.cpp
--- a/Engine/Source/Editor/BlueprintGraph/Private/EdGraphSchema_K2.cpp
+++ b/Engine/Source/Editor/BlueprintGraph/Private/EdGraphSchema_K2.cpp
@@ -910,9 +910,10 @@ bool UEdGraphSchema_K2::CanFunctionBeUsedInGraph(const UClass* InClass, const UF
 				{
 					UBlueprint* BP = FBlueprintEditorUtils::FindBlueprintForGraph(InDestGraph);
 					const bool bIsFunctLib = BP && (EBlueprintType::BPTYPE_FunctionLibrary == BP->BlueprintType);
+					const bool bIsMacroLib = BP && (EBlueprintType::BPTYPE_MacroLibrary == BP->BlueprintType);
 					UClass* ParentClass = BP ? BP->ParentClass : NULL;
 					const bool bIncompatibleParent = ParentClass && (!FBlueprintEditorUtils::ImplementsGetWorld(BP) && !ParentClass->HasMetaDataHierarchical(FBlueprintMetadata::MD_ShowWorldContextPin));
-					if (!bIsFunctLib && bIncompatibleParent)
+					if (!bIsMacroLib && !bIsFunctLib && bIncompatibleParent)
 					{
 						if (OutReason != nullptr)
 						{
diff --git a/Engine/Source/Editor/BlueprintGraph/Private/K2Node_CallFunction.cpp b/Engine/Source/Editor/BlueprintGraph/Private/K2Node_CallFunction.cpp
--- a/Engine/Source/Editor/BlueprintGraph/Private/K2Node_CallFunction.cpp
+++ b/Engine/Source/Editor/BlueprintGraph/Private/K2Node_CallFunction.cpp
@@ -2345,9 +2345,20 @@ void UK2Node_CallFunction::ValidateNodeDuringCompilation(class FCompilerResultsL
 
 		// enforce WorldContext restrictions
 		const bool bInsideBpFuncLibrary = Blueprint && (BPTYPE_FunctionLibrary == Blueprint->BlueprintType);
-		if (!bInsideBpFuncLibrary && 
+		
+		// go through all of the pins and verify if we have a visible world context or not.
+		bool bVisibleWorldContext = false;
+		for (const UEdGraphPin* Pin : Pins)
+		{
+			if (!Pin->bHidden && Pin->PinName.IsEqual(FName("WorldContextObject")))
+			{
+				bVisibleWorldContext = true;
+				break;
+			}
+		}
+		if (!bInsideBpFuncLibrary &&
 			Function->HasMetaData(FBlueprintMetadata::MD_WorldContext) && 
-			!Function->HasMetaData(FBlueprintMetadata::MD_CallableWithoutWorldContext))
+			(!Function->HasMetaData(FBlueprintMetadata::MD_CallableWithoutWorldContext) && !bVisibleWorldContext))
 		{
 			check(Blueprint);
 			UClass* ParentClass = Blueprint->ParentClass;
