Deprecated bIsBeadFunction.

This was a really old feature that never took off. It's not needed anymore.

diff --git a/Engine/Source/Editor/BlueprintGraph/Classes/K2Node.h b/Engine/Source/Editor/BlueprintGraph/Classes/K2Node.h
--- a/Engine/Source/Editor/BlueprintGraph/Classes/K2Node.h
+++ b/Engine/Source/Editor/BlueprintGraph/Classes/K2Node.h
@@ -266,7 +266,7 @@ public:
 	 */
 	virtual bool IsActionFilteredOut(class FBlueprintActionFilter const& Filter) { return false; }
 
-	/** Should draw as a bead with no location of it's own */
+	UE_DEPRECATED(5.4, "ShouldDrawAsBead is deprecated")
 	virtual bool ShouldDrawAsBead() const { return false; }
 
 	/** Return whether the node's properties display in the blueprint details panel */
diff --git a/Engine/Source/Editor/BlueprintGraph/Classes/K2Node_CallFunction.h b/Engine/Source/Editor/BlueprintGraph/Classes/K2Node_CallFunction.h
--- a/Engine/Source/Editor/BlueprintGraph/Classes/K2Node_CallFunction.h
+++ b/Engine/Source/Editor/BlueprintGraph/Classes/K2Node_CallFunction.h
@@ -59,8 +59,7 @@ class BLUEPRINTGRAPH_API UK2Node_CallFunction : public UK2Node
 	UPROPERTY()
 	uint32 bIsFinalFunction:1;
 
-	/** Indicates that this is a 'bead' function with no fixed location; it is drawn between the nodes that it is wired to */
-	UPROPERTY()
+	UE_DEPRECATED(5.4, "bIsBeadFunction is deprecated")
 	uint32 bIsBeadFunction:1;
 
 	/** The function to call */
@@ -125,7 +124,8 @@ public:
 	virtual bool IsNodePure() const override { return bIsPureFunc; }
 	virtual void PostReconstructNode() override;
 	virtual bool ShouldDrawCompact() const override;
-	virtual bool ShouldDrawAsBead() const override;
+	UE_DEPRECATED(5.4, "ShouldDrawAsBead is deprecated")
+	virtual bool ShouldDrawAsBead() const override { return false; }
 	virtual FText GetCompactNodeTitle() const override;
 	virtual void PostPasteNode() override;
 	virtual bool CanSplitPin(const UEdGraphPin* Pin) const override;
@@ -165,7 +165,7 @@ public:
 	 * @return	Pointer to the pin that was created
 	 */
 	virtual UEdGraphPin* CreateSelfPin(const UFunction* Function);
-	
+
 	/**
 	 * Creates all of the pins required to call a particular UFunction.
 	 *
diff --git a/Engine/Source/Editor/BlueprintGraph/Private/K2Node_CallFunction.cpp b/Engine/Source/Editor/BlueprintGraph/Private/K2Node_CallFunction.cpp
--- a/Engine/Source/Editor/BlueprintGraph/Private/K2Node_CallFunction.cpp
+++ b/Engine/Source/Editor/BlueprintGraph/Private/K2Node_CallFunction.cpp
@@ -1341,15 +1341,6 @@ void UK2Node_CallFunction::NotifyPinConnectionListChanged(UEdGraphPin* Pin)
 		GetGraph()->NotifyGraphChanged();
 	}
 
-	if (bIsBeadFunction)
-	{
-		if (Pin->LinkedTo.Num() == 0)
-		{
-			// Commit suicide; bead functions must always have an input and output connection
-			DestroyNode();
-		}
-	}
-
 	InvalidatePinTooltips();
 	if(!Pin->IsPendingKill())
 	{
@@ -2000,11 +1991,6 @@ bool UK2Node_CallFunction::ShouldDrawCompact() const
 	return ShouldDrawCompact(Function);
 }
 
-bool UK2Node_CallFunction::ShouldDrawAsBead() const
-{
-	return bIsBeadFunction;
-}
-
 bool UK2Node_CallFunction::ShouldShowNodeProperties() const
 {
 	// Show node properties if this corresponds to a function graph
diff --git a/Engine/Source/Editor/GraphEditor/Private/KismetNodes/SGraphNodeK2Base.cpp b/Engine/Source/Editor/GraphEditor/Private/KismetNodes/SGraphNodeK2Base.cpp
--- a/Engine/Source/Editor/GraphEditor/Private/KismetNodes/SGraphNodeK2Base.cpp
+++ b/Engine/Source/Editor/GraphEditor/Private/KismetNodes/SGraphNodeK2Base.cpp
@@ -440,14 +440,6 @@ void SGraphNodeK2Base::UpdateGraphNode()
 	}
 }
 
-bool SGraphNodeK2Base::RequiresSecondPassLayout() const
-{
-	UK2Node* K2Node = CastChecked<UK2Node>(GraphNode);
-	const bool bBeadMode = K2Node->ShouldDrawAsBead();
-
-	return bBeadMode;
-}
-
 FText SGraphNodeK2Base::GetNodeCompactTitle() const
 {
 	UK2Node* K2Node = CastChecked<UK2Node>(GraphNode);
diff --git a/Engine/Source/Editor/GraphEditor/Public/KismetNodes/SGraphNodeK2Base.h b/Engine/Source/Editor/GraphEditor/Public/KismetNodes/SGraphNodeK2Base.h
--- a/Engine/Source/Editor/GraphEditor/Public/KismetNodes/SGraphNodeK2Base.h
+++ b/Engine/Source/Editor/GraphEditor/Public/KismetNodes/SGraphNodeK2Base.h
@@ -32,7 +32,6 @@ public:
 	virtual void UpdateGraphNode() override;
 
 	//~ Begin SNodePanel::SNode Interface
-	virtual bool RequiresSecondPassLayout() const override;
 	virtual void GetOverlayBrushes(bool bSelected, const FVector2D WidgetSize, TArray<FOverlayBrushInfo>& Brushes) const override;
 	virtual void GetNodeInfoPopups(FNodeInfoContext* Context, TArray<FGraphInformationPopupInfo>& Popups) const override;
 	virtual const FSlateBrush* GetShadowBrush(bool bSelected) const override;
diff --git a/Engine/Source/Editor/Kismet/Private/BlueprintEditor.cpp b/Engine/Source/Editor/Kismet/Private/BlueprintEditor.cpp
--- a/Engine/Source/Editor/Kismet/Private/BlueprintEditor.cpp
+++ b/Engine/Source/Editor/Kismet/Private/BlueprintEditor.cpp
@@ -4738,7 +4738,7 @@ void FBlueprintEditor::DoPromoteToVariable( UBlueprint* InBlueprint, UEdGraphPin
 			// Ignore names from compact nodes that don't usually display the pin names
 			if (const UK2Node* K2Node = Cast<UK2Node>(InTargetPin->GetOwningNode()))
 			{
-				if (K2Node->ShouldDrawCompact() || K2Node->ShouldDrawAsBead())
+				if (K2Node->ShouldDrawCompact())
 				{
 					IdealVarName.Empty();
 				}
