Deprecated bIsFinalFunction.

"final" was an UnrealScript concept that never made it into Blueprints. Currently, it's synonymous with the CallParentFunction node. Instead of using this flag, we can simply check the node type directly.

We still might want to revisit this for Blueprints, but we likely won't use this bool.

diff --git a/Engine/Source/Editor/BlueprintGraph/Classes/K2Node_CallFunction.h b/Engine/Source/Editor/BlueprintGraph/Classes/K2Node_CallFunction.h
--- a/Engine/Source/Editor/BlueprintGraph/Classes/K2Node_CallFunction.h
+++ b/Engine/Source/Editor/BlueprintGraph/Classes/K2Node_CallFunction.h
@@ -55,8 +55,7 @@ class BLUEPRINTGRAPH_API UK2Node_CallFunction : public UK2Node
 	UPROPERTY()
 	uint32 bIsInterfaceCall:1;
 
-	/** Indicates that this is a call to a final / superclass's function */
-	UPROPERTY()
+	UE_DEPRECATED(5.4, "bIsFinalFunction is deprecated.")
 	uint32 bIsFinalFunction:1;
 
 	UE_DEPRECATED(5.4, "bIsBeadFunction is deprecated")
diff --git a/Engine/Source/Editor/BlueprintGraph/Private/CallFunctionHandler.cpp b/Engine/Source/Editor/BlueprintGraph/Private/CallFunctionHandler.cpp
--- a/Engine/Source/Editor/BlueprintGraph/Private/CallFunctionHandler.cpp
+++ b/Engine/Source/Editor/BlueprintGraph/Private/CallFunctionHandler.cpp
@@ -480,7 +480,7 @@ void FKCHandler_CallFunction::CreateFunctionCallStatement(FKismetFunctionContext
 				Statement.FunctionContext = Target;
 				Statement.Type = KCST_CallFunction;
 				Statement.bIsInterfaceContext = IsCalledFunctionFromInterface(Node) || bIsInterfaceContextTerm;
-				Statement.bIsParentContext = IsCalledFunctionFinal(Node);
+				Statement.bIsParentContext = Node->IsA<UK2Node_CallParentFunction>();
 
 				Statement.LHS = LHSTerm;
 				Statement.RHS = RHSTerms;
@@ -559,28 +559,20 @@ UClass* FKCHandler_CallFunction::GetCallingContext(FKismetFunctionContext& Conte
 {
 	// Find the calling scope
 	UClass* SearchScope = Context.NewClass;
-	UK2Node_CallFunction* CallFuncNode = Cast<UK2Node_CallFunction>(Node);
-	if (CallFuncNode && CallFuncNode->bIsFinalFunction)
+	
+	if (UK2Node_CallParentFunction* ParentCall = Cast<UK2Node_CallParentFunction>(Node))
 	{
-		if (UK2Node_CallParentFunction* ParentCall = Cast<UK2Node_CallParentFunction>(Node))
-		{
-			// Special Case:  super call functions should search up their class hierarchy, and find the first legitimate implementation of the function
-			const FName FuncName = CallFuncNode->FunctionReference.GetMemberName();
-			UClass* SearchContext = Context.NewClass->GetSuperClass();
-
-			UFunction* ParentFunc = nullptr;
-			if (SearchContext)
-			{
-				ParentFunc = SearchContext->FindFunctionByName(FuncName);
-			}
+		// Special Case: super call functions should search up their class hierarchy, and find the first legitimate implementation of the function
+		const FName FuncName = ParentCall->FunctionReference.GetMemberName();
+		UClass* SearchContext = Context.NewClass->GetSuperClass();
 
-			return ParentFunc ? ParentFunc->GetOuterUClass() : nullptr;
-		}
-		else
+		UFunction* ParentFunc = nullptr;
+		if (SearchContext)
 		{
-			// Final functions need the call context to be the specified class, so don't bother checking for the self pin.   The schema should enforce this.
-			return CallFuncNode->FunctionReference.GetMemberParentClass(CallFuncNode->GetBlueprintClassFromNode());
+			ParentFunc = SearchContext->FindFunctionByName(FuncName);
 		}
+
+		return ParentFunc ? ParentFunc->GetOuterUClass() : nullptr;
 	}
 	else
 	{
@@ -1019,7 +1011,8 @@ void FKCHandler_CallFunction::CheckIfFunctionIsCallable(UFunction* Function, FKi
 	// Verify that the function is a Blueprint callable function (in case a BlueprintCallable specifier got removed)
 	if (!Function->HasAnyFunctionFlags(FUNC_BlueprintCallable) && (Function->GetOuter() != Context.NewClass))
 	{
-		if (!IsCalledFunctionFinal(Node) && Function->GetName().Find(UEdGraphSchema_K2::FN_ExecuteUbergraphBase.ToString()))
+		const bool bIsParentFunction = Node && Node->IsA<UK2Node_CallParentFunction>();
+		if (!bIsParentFunction && Function->GetName().Find(UEdGraphSchema_K2::FN_ExecuteUbergraphBase.ToString()))
 		{
 			CompilerContext.MessageLog.Error(*FText::Format(NSLOCTEXT("KismetCompiler", "ShouldNotCallFromBlueprint_ErrorFmt", "Function '{0}' called from @@ should not be called from a Blueprint"), FText::FromString(Function->GetName())).ToString(), Node);
 		}
diff --git a/Engine/Source/Editor/BlueprintGraph/Private/K2Node_CallParentFunction.cpp b/Engine/Source/Editor/BlueprintGraph/Private/K2Node_CallParentFunction.cpp
--- a/Engine/Source/Editor/BlueprintGraph/Private/K2Node_CallParentFunction.cpp
+++ b/Engine/Source/Editor/BlueprintGraph/Private/K2Node_CallParentFunction.cpp
@@ -25,7 +25,6 @@
 UK2Node_CallParentFunction::UK2Node_CallParentFunction(const FObjectInitializer& ObjectInitializer)
 	: Super(ObjectInitializer)
 {
-	bIsFinalFunction = true;
 }
 
 FText UK2Node_CallParentFunction::GetNodeTitle(ENodeTitleType::Type TitleType) const
diff --git a/Engine/Source/Editor/BlueprintGraph/Public/CallFunctionHandler.h b/Engine/Source/Editor/BlueprintGraph/Public/CallFunctionHandler.h
--- a/Engine/Source/Editor/BlueprintGraph/Public/CallFunctionHandler.h
+++ b/Engine/Source/Editor/BlueprintGraph/Public/CallFunctionHandler.h
@@ -37,28 +37,27 @@ public:
 
 	bool IsCalledFunctionPure(UEdGraphNode* Node)
 	{
-		if(UK2Node_CallFunction* CallFunctionNode = Cast<UK2Node_CallFunction>(Node))
+		if (UK2Node_CallFunction* CallFunctionNode = Cast<UK2Node_CallFunction>(Node))
 		{
 			return CallFunctionNode->bIsPureFunc;
 		}
+
 		return false;
 	}
 
+	UE_DEPRECATED(5.4, "IsCalledFunctionFinal is deprecated")
 	bool IsCalledFunctionFinal(UEdGraphNode* Node)
 	{
-		if(UK2Node_CallFunction* CallFunctionNode = Cast<UK2Node_CallFunction>(Node))
-		{
-			return CallFunctionNode->bIsFinalFunction;
-		}
 		return false;
 	}
 
 	bool IsCalledFunctionFromInterface(UEdGraphNode* Node)
 	{
-		if(UK2Node_CallFunction* CallFunctionNode = Cast<UK2Node_CallFunction>(Node))
+		if (UK2Node_CallFunction* CallFunctionNode = Cast<UK2Node_CallFunction>(Node))
 		{
 			return CallFunctionNode->bIsInterfaceCall;
 		}
+
 		return false;
 	}
 
