PR #10957: Fix GetSubsystem blueprint node to correctly handle when ShowWorldContextPin is set on a subsystem's parent class

diff --git a/Engine/Source/Editor/BlueprintGraph/Private/K2Node_GetSubsystem.cpp b/Engine/Source/Editor/BlueprintGraph/Private/K2Node_GetSubsystem.cpp
--- a/Engine/Source/Editor/BlueprintGraph/Private/K2Node_GetSubsystem.cpp
+++ b/Engine/Source/Editor/BlueprintGraph/Private/K2Node_GetSubsystem.cpp
@@ -73,7 +73,7 @@ void UK2Node_GetSubsystem::Initialize( UClass* NodeClass )
 void UK2Node_GetSubsystem::AllocateDefaultPins()
 {
 	// If required add the world context pin
-	if (GetBlueprint()->ParentClass->HasMetaData(FBlueprintMetadata::MD_ShowWorldContextPin))
+	if (GetBlueprint()->ParentClass->HasMetaDataHierarchical(FBlueprintMetadata::MD_ShowWorldContextPin))
 	{
 		CreatePin(EGPD_Input, UEdGraphSchema_K2::PC_Object, UObject::StaticClass(), TEXT("WorldContext"));
 	}
