Do not show the "breakpoints" category in the right click context menu for K2Nodes if they cannot place a breakpoint. This prevents a useless empty category from being added.

diff --git a/Engine/Source/Editor/BlueprintGraph/Private/EdGraphSchema_K2.cpp b/Engine/Source/Editor/BlueprintGraph/Private/EdGraphSchema_K2.cpp
--- a/Engine/Source/Editor/BlueprintGraph/Private/EdGraphSchema_K2.cpp
+++ b/Engine/Source/Editor/BlueprintGraph/Private/EdGraphSchema_K2.cpp
@@ -1881,6 +1881,7 @@ void UEdGraphSchema_K2::GetContextMenuActions(UToolMenu* Menu, UGraphNodeContext
 					}
 					
 					// Add breakpoint actions
+					if (K2Node->CanPlaceBreakpoints())
 					{
 						FToolMenuSection& Section = Menu->AddSection("EdGraphSchemaBreakpoints", LOCTEXT("BreakpointsHeader", "Breakpoints"));
 						Section.AddMenuEntry( FGraphEditorCommands::Get().ToggleBreakpoint );
