Notify when the pins are removed from the FormatText k2node. That occurs when the format default value was valid and link the format pin.

diff --git a/Engine/Source/Editor/BlueprintGraph/Private/K2Node_FormatText.cpp b/Engine/Source/Editor/BlueprintGraph/Private/K2Node_FormatText.cpp
--- a/Engine/Source/Editor/BlueprintGraph/Private/K2Node_FormatText.cpp
+++ b/Engine/Source/Editor/BlueprintGraph/Private/K2Node_FormatText.cpp
@@ -168,6 +168,7 @@ void UK2Node_FormatText::PinConnectionListChanged(UEdGraphPin* Pin)
 		PinNames.Empty();
 		GetSchema()->TrySetDefaultText(*FormatPin, FText::GetEmpty());
 
+		bool bRemoved = false;
 		for(auto It = Pins.CreateConstIterator(); It; ++It)
 		{
 			UEdGraphPin* CheckPin = *It;
@@ -176,10 +177,16 @@ void UK2Node_FormatText::PinConnectionListChanged(UEdGraphPin* Pin)
 				CheckPin->Modify();
 				CheckPin->MarkAsGarbage();
 				Pins.Remove(CheckPin);
+				bRemoved = true;
 				--It;
 			}
 		}
 
+		if (bRemoved)
+		{
+			GetGraph()->NotifyNodeChanged(this);
+		}
+
 		FBlueprintEditorUtils::MarkBlueprintAsStructurallyModified(GetBlueprint());
 	}
 
