Add function name / asset name to ensure

diff --git a/Engine/Source/Editor/BlueprintGraph/Private/K2Node_FunctionEntry.cpp b/Engine/Source/Editor/BlueprintGraph/Private/K2Node_FunctionEntry.cpp
--- a/Engine/Source/Editor/BlueprintGraph/Private/K2Node_FunctionEntry.cpp
+++ b/Engine/Source/Editor/BlueprintGraph/Private/K2Node_FunctionEntry.cpp
@@ -403,7 +403,11 @@ void UK2Node_FunctionEntry::AllocateDefaultPins()
 	Super::AllocateDefaultPins();
 
 	if (FFunctionEntryHelper::RequireWorldContextParameter(this) 
-		&& ensure(!FindPin(FFunctionEntryHelper::GetWorldContextPinName())))
+		&& ensureMsgf(!FindPin(FFunctionEntryHelper::GetWorldContextPinName()), 
+		TEXT("%s: World context parameter pin already exiss on function entry node %s"), 
+			*GetOutermost()->GetName(),
+			*(CustomGeneratedFunctionName.IsNone() ? FunctionReference.GetMemberName() : CustomGeneratedFunctionName).ToString()
+		))
 	{
 		UEdGraphPin* WorldContextPin = CreatePin(
 			EGPD_Output,
