More pin type menu optimzations, avoid wandering back and forth between SoftObjectPath and AssetData

diff --git a/Engine/Source/Editor/BlueprintGraph/Classes/EdGraphSchema_K2.h b/Engine/Source/Editor/BlueprintGraph/Classes/EdGraphSchema_K2.h
--- a/Engine/Source/Editor/BlueprintGraph/Classes/EdGraphSchema_K2.h
+++ b/Engine/Source/Editor/BlueprintGraph/Classes/EdGraphSchema_K2.h
@@ -447,8 +447,8 @@ public:
 		FEdGraphPinType PinType;
 		uint8 PossibleObjectReferenceTypes;
 
-		/** Asset Reference, used when PinType.PinSubCategoryObject is not loaded yet */
-		FSoftObjectPath SubCategoryObjectAssetReference;
+		/** Asset Data, used when PinType.PinSubCategoryObject is not loaded yet */
+		FAssetData CachedAssetData;
 
 		FText CachedDescription;
 
@@ -475,18 +475,8 @@ public:
 
 		FPinTypeTreeInfo(const FText& InFriendlyName, const FName CategoryName, const UEdGraphSchema_K2* Schema, const FText& InTooltip, bool bInReadOnly = false, FTypesDatabase* TypesDatabase = nullptr);
 		FPinTypeTreeInfo(const FName CategoryName, UObject* SubCategoryObject, const FText& InTooltip, bool bInReadOnly = false, uint8 InPossibleObjectReferenceTypes = 0);
-		FPinTypeTreeInfo(const FText& InFriendlyName, const FName CategoryName, const FSoftObjectPath& SubCategoryObject, const FText& InTooltip, bool bInReadOnly = false, uint8 InPossibleObjectReferenceTypes = 0);
-
-		FPinTypeTreeInfo(TSharedPtr<FPinTypeTreeInfo> InInfo)
-		{
-			PinType = InInfo->PinType;
-			bReadOnly = InInfo->bReadOnly;
-			FriendlyName = InInfo->FriendlyName;
-			Tooltip = InInfo->Tooltip;
-			SubCategoryObjectAssetReference = InInfo->SubCategoryObjectAssetReference;
-			CachedDescription = InInfo->CachedDescription;
-			PossibleObjectReferenceTypes = InInfo->PossibleObjectReferenceTypes;
-		}
+		FPinTypeTreeInfo(const FText& InFriendlyName, const FName CategoryName, const FAssetData& AssetData, const FText& InTooltip, bool bInReadOnly = false, uint8 InPossibleObjectReferenceTypes = 0);
+		FPinTypeTreeInfo(TSharedPtr<FPinTypeTreeInfo> InInfo);
 		
 		/** Returns a succinct menu description of this type */
 		FText GetDescription() const;
@@ -511,10 +501,7 @@ public:
 			return PossibleObjectReferenceTypes;
 		}
 
-		const FSoftObjectPath& GetSubCategoryObjectAsset() const
-		{
-			return SubCategoryObjectAssetReference;
-		}
+		const FAssetData& GetCachedAssetData() const;
 
 	private:
 
diff --git a/Engine/Source/Editor/BlueprintGraph/Private/EdGraphSchema_K2.cpp b/Engine/Source/Editor/BlueprintGraph/Private/EdGraphSchema_K2.cpp
--- a/Engine/Source/Editor/BlueprintGraph/Private/EdGraphSchema_K2.cpp
+++ b/Engine/Source/Editor/BlueprintGraph/Private/EdGraphSchema_K2.cpp
@@ -295,7 +295,7 @@ UEdGraphSchema_K2::FPinTypeTreeInfo::FPinTypeTreeInfo(const FText& InFriendlyNam
 
 struct FUnloadedAssetData
 {
-	FSoftObjectPath SoftObjectPath;
+	FAssetData AssetData;
 	FText AssetFriendlyName;
 	FText Tooltip;
 	uint8 PossibleObjectReferenceTypes;
@@ -305,7 +305,7 @@ struct FUnloadedAssetData
 	{}
 
 	FUnloadedAssetData(const FAssetData& InAsset, uint8 InPossibleObjectReferenceTypes = 0)
-		: SoftObjectPath(InAsset.ToSoftObjectPath())
+		: AssetData(InAsset)
 		, AssetFriendlyName(FText::FromString(FName::NameToDisplayString(InAsset.AssetName.ToString(), false)))
 		, PossibleObjectReferenceTypes(InPossibleObjectReferenceTypes)
 	{
@@ -552,7 +552,7 @@ public:
 			{
 				FPinTypeTreeInfoPtr TypeTreeInfo = MakeShareable(new UEdGraphSchema_K2::FPinTypeTreeInfo(It.AssetFriendlyName
 					, CategoryName
-					, It.SoftObjectPath
+					, It.AssetData
 					, It.Tooltip
 					, false
 					, It.PossibleObjectReferenceTypes));
@@ -567,13 +567,12 @@ public:
 const FEdGraphPinType& UEdGraphSchema_K2::FPinTypeTreeInfo::GetPinType(bool bForceLoadedSubCategoryObject)
 {
 	// Only attempt to load the sub category object if we need to
-	if (SubCategoryObjectAssetReference.IsValid() && (!PinType.PinSubCategoryObject.IsValid() || FSoftObjectPath(PinType.PinSubCategoryObject.Get()) != SubCategoryObjectAssetReference))
+	if (CachedAssetData.IsValid() && (!PinType.PinSubCategoryObject.IsValid() || FSoftObjectPath(PinType.PinSubCategoryObject.Get()) != CachedAssetData.GetSoftObjectPath()))
 	{
 		const FAssetRegistryModule& AssetRegistryModule = FModuleManager::LoadModuleChecked<FAssetRegistryModule>(TEXT("AssetRegistry"));
-		const FAssetData AssetData = AssetRegistryModule.Get().GetAssetByObjectPath(SubCategoryObjectAssetReference);
-		if (bForceLoadedSubCategoryObject || AssetData.IsAssetLoaded())
+		if (bForceLoadedSubCategoryObject || CachedAssetData.IsAssetLoaded())
 		{
-			UObject* LoadedObject = AssetData.GetAsset();
+			UObject* LoadedObject = CachedAssetData.GetAsset();
 
 			if (UBlueprint* BlueprintObject = Cast<UBlueprint>(LoadedObject))
 			{
@@ -628,23 +627,34 @@ UEdGraphSchema_K2::FPinTypeTreeInfo::FPinTypeTreeInfo(const FName CategoryName,
 	CachedDescription = GenerateDescription();
 }
 
-UEdGraphSchema_K2::FPinTypeTreeInfo::FPinTypeTreeInfo(const FText& InFriendlyName, const FName CategoryName, const FSoftObjectPath& SubCategoryObject, const FText& InTooltip, bool bInReadOnly, uint8 InPossibleObjectReferenceTypes)
+UEdGraphSchema_K2::FPinTypeTreeInfo::FPinTypeTreeInfo(const FText& InFriendlyName, const FName CategoryName, const FAssetData& AssetData, const FText& InTooltip, bool bInReadOnly, uint8 InPossibleObjectReferenceTypes)
 	: PossibleObjectReferenceTypes(InPossibleObjectReferenceTypes)
 {
 	FriendlyName = InFriendlyName;
 
 	check(!CategoryName.IsNone());
-	check(SubCategoryObject.IsValid());
+	check(AssetData.IsValid());
 
 	Tooltip = InTooltip;
 	PinType.PinCategory = CategoryName;
 
-	SubCategoryObjectAssetReference = SubCategoryObject;
+	CachedAssetData = AssetData;
 
 	bReadOnly = bInReadOnly;
 	CachedDescription = GenerateDescription();
 }
 
+UEdGraphSchema_K2::FPinTypeTreeInfo::FPinTypeTreeInfo(TSharedPtr<FPinTypeTreeInfo> InInfo)
+{
+	PinType = InInfo->PinType;
+	bReadOnly = InInfo->bReadOnly;
+	FriendlyName = InInfo->FriendlyName;
+	CachedAssetData = InInfo->CachedAssetData;
+	Tooltip = InInfo->Tooltip;
+	CachedDescription = InInfo->CachedDescription;
+	PossibleObjectReferenceTypes = InInfo->PossibleObjectReferenceTypes;
+}
+
 FText UEdGraphSchema_K2::FPinTypeTreeInfo::GenerateDescription()
 {
 	if (!FriendlyName.IsEmpty())
@@ -676,6 +686,11 @@ FText UEdGraphSchema_K2::FPinTypeTreeInfo::GetDescription() const
 	return CachedDescription;
 }
 
+const FAssetData& UEdGraphSchema_K2::FPinTypeTreeInfo::GetCachedAssetData() const
+{
+	return CachedAssetData;
+}
+
 const FName UEdGraphSchema_K2::PC_Exec(TEXT("exec"));
 const FName UEdGraphSchema_K2::PC_Boolean(TEXT("bool"));
 const FName UEdGraphSchema_K2::PC_Byte(TEXT("byte"));
diff --git a/Engine/Source/Editor/Kismet/Private/BlueprintEditor.cpp b/Engine/Source/Editor/Kismet/Private/BlueprintEditor.cpp
--- a/Engine/Source/Editor/Kismet/Private/BlueprintEditor.cpp
+++ b/Engine/Source/Editor/Kismet/Private/BlueprintEditor.cpp
@@ -419,12 +419,11 @@ namespace BlueprintEditorImpl
 		// IPinTypeSelectorFilter interface
 		virtual bool ShouldShowPinTypeTreeItem(FPinTypeTreeItem InItem) const override
 		{
-			const FSoftObjectPath& AssetReference = InItem->GetSubCategoryObjectAsset();
 			FTopLevelAssetPath TopLevelAssetPath;
-
-			if (AssetReference.IsAsset())
+			const FAssetData& AssetData = InItem->GetCachedAssetData();
+			if (AssetData.IsValid())
 			{
-				TopLevelAssetPath = FTopLevelAssetPath(AssetReference.GetLongPackageFName(), *AssetReference.GetAssetName());
+				TopLevelAssetPath = FTopLevelAssetPath(AssetData.PackageName, AssetData.AssetName);
 			}
 
 			// First check pin type permissions
@@ -434,10 +433,8 @@ namespace BlueprintEditorImpl
 			}
 
 			// Then asset permissions
-			if(AssetReferenceFilter.IsValid() && AssetReference.IsValid())
+			if(AssetReferenceFilter.IsValid() && AssetData.IsValid())
 			{
-				FAssetRegistryModule& AssetRegistryModule = FModuleManager::LoadModuleChecked<FAssetRegistryModule>("AssetRegistry");
-				FAssetData AssetData = AssetRegistryModule.Get().GetAssetByObjectPath(AssetReference.GetWithoutSubPath());
 				if (!AssetReferenceFilter->PassesFilter(AssetData))
 				{
 					return false;
diff --git a/Engine/Source/Editor/Kismet/Private/BlueprintNamespaceHelper.cpp b/Engine/Source/Editor/Kismet/Private/BlueprintNamespaceHelper.cpp
--- a/Engine/Source/Editor/Kismet/Private/BlueprintNamespaceHelper.cpp
+++ b/Engine/Source/Editor/Kismet/Private/BlueprintNamespaceHelper.cpp
@@ -111,8 +111,7 @@ public:
 		}
 		else
 		{
-			const FSoftObjectPath& AssetRef = InItem->GetSubCategoryObjectAsset();
-			if (AssetRef.IsValid() && !CachedNamespaceHelper->IsImportedObject(AssetRef))
+			if (!CachedNamespaceHelper->IsImportedAsset(InItem->GetCachedAssetData()))
 			{
 				// A pin type whose underlying asset may be either loaded or unloaded, but is not imported.
 				return false;
diff --git a/Engine/Source/Editor/Kismet/Private/BlueprintNamespaceUtilities.cpp b/Engine/Source/Editor/Kismet/Private/BlueprintNamespaceUtilities.cpp
--- a/Engine/Source/Editor/Kismet/Private/BlueprintNamespaceUtilities.cpp
+++ b/Engine/Source/Editor/Kismet/Private/BlueprintNamespaceUtilities.cpp
@@ -60,25 +60,17 @@ FString FBlueprintNamespaceUtilities::GetAssetNamespace(const FAssetData& InAsse
 
 	if (InAssetData.IsValid())
 	{
-		if (const UObject* AssetObject = InAssetData.FastGetAsset())
+		using namespace UE::Editor::Kismet::Private;
+
+		// @todo_namespaces - Add cases for unloaded UDS/UDE assets once they have a searchable namespace tag or property.
+		FString TagValue;
+		if (InAssetData.GetTagValue<FString>(GET_MEMBER_NAME_STRING_CHECKED(UBlueprint, BlueprintNamespace), TagValue))
 		{
-			Namespace = GetObjectNamespace(AssetObject);
+			Namespace = MoveTemp(TagValue);
 		}
-		else
+		else if (DefaultBlueprintNamespaceType == EDefaultBlueprintNamespaceType::UsePackagePathAsDefaultNamespace)
 		{
-			using namespace UE::Editor::Kismet::Private;
-
-			// @todo_namespaces - Add cases for unloaded UDS/UDE assets once they have a searchable namespace tag or property.
-
-			FString TagValue;
-			if (InAssetData.GetTagValue<FString>(GET_MEMBER_NAME_STRING_CHECKED(UBlueprint, BlueprintNamespace), TagValue))
-			{
-				Namespace = MoveTemp(TagValue);
-			}
-			else if (DefaultBlueprintNamespaceType == EDefaultBlueprintNamespaceType::UsePackagePathAsDefaultNamespace)
-			{
-				ConvertPackagePathToNamespacePath(InAssetData.PackageName.ToString(), Namespace);
-			}
+			ConvertPackagePathToNamespacePath(InAssetData.PackageName.ToString(), Namespace);
 		}
 	}
 
