One last cleanup of the pin type selector, avoid repeated FText::Rebuild calls to sort the pin lists, clarify the intent of FPinTypeTreeInfo constructors

diff --git a/Engine/Source/Editor/BlueprintGraph/Classes/EdGraphSchema_K2.h b/Engine/Source/Editor/BlueprintGraph/Classes/EdGraphSchema_K2.h
--- a/Engine/Source/Editor/BlueprintGraph/Classes/EdGraphSchema_K2.h
+++ b/Engine/Source/Editor/BlueprintGraph/Classes/EdGraphSchema_K2.h
@@ -450,8 +450,11 @@ public:
 		/** Asset Data, used when PinType.PinSubCategoryObject is not loaded yet */
 		FAssetData CachedAssetData;
 
+		/** The pin type description, localized */
 		FText CachedDescription;
 
+		/** A copy of the localized CachedDescription string, for sorting */
+		TSharedPtr<FString> CachedDescriptionString;
 	public:
 		/** The children of this pin type */
 		TArray< TSharedPtr<FPinTypeTreeInfo> > Children;
@@ -459,9 +462,6 @@ public:
 		/** Whether or not this pin type is selectable as an actual type, or is just a category, with some subtypes */
 		bool bReadOnly;
 
-		/** Friendly display name of pin type; also used to see if it has subtypes */
-		FText FriendlyName;
-
 		/** Text for regular tooltip */
 		FText Tooltip;
 
@@ -479,7 +479,13 @@ public:
 		FPinTypeTreeInfo(TSharedPtr<FPinTypeTreeInfo> InInfo);
 		
 		/** Returns a succinct menu description of this type */
-		FText GetDescription() const;
+		const FText& GetDescription() const;
+		
+		/** Returns the localized description as a string, for faster sorting */
+		const FString& GetCachedDescriptionString() const
+		{
+			return *CachedDescriptionString;
+		}
 
 		FText GetToolTip() const
 		{
@@ -510,8 +516,6 @@ public:
 			, bReadOnly(false)
 		{}
 
-		void Init(const FText& FriendlyCategoryName, const FName CategoryName, const UEdGraphSchema_K2* Schema, const FText& InTooltip, bool bInReadOnly, FTypesDatabase* TypesDatabase);
-
 		FText GenerateDescription();
 	};
 
diff --git a/Engine/Source/Editor/BlueprintGraph/Private/EdGraphSchema_K2.cpp b/Engine/Source/Editor/BlueprintGraph/Private/EdGraphSchema_K2.cpp
--- a/Engine/Source/Editor/BlueprintGraph/Private/EdGraphSchema_K2.cpp
+++ b/Engine/Source/Editor/BlueprintGraph/Private/EdGraphSchema_K2.cpp
@@ -287,12 +287,6 @@ const FName FBlueprintMetadata::MD_FieldNotify(TEXT("FieldNotify"));
 
 #define LOCTEXT_NAMESPACE "KismetSchema"
 
-UEdGraphSchema_K2::FPinTypeTreeInfo::FPinTypeTreeInfo(const FText& InFriendlyName, const FName CategoryName, const UEdGraphSchema_K2* Schema, const FText& InTooltip, bool bInReadOnly/*=false*/, FTypesDatabase* TypesDatabase /*=nullptr*/)
-	: PossibleObjectReferenceTypes(0)
-{
-	Init(InFriendlyName, CategoryName, Schema, InTooltip, bInReadOnly, TypesDatabase);
-}
-
 struct FUnloadedAssetData
 {
 	FAssetData AssetData;
@@ -350,13 +344,6 @@ class FGatherTypesHelper
 {
 private:
 	typedef TSharedPtr<UEdGraphSchema_K2::FPinTypeTreeInfo> FPinTypeTreeInfoPtr;
-	struct FCompareChildren
-	{
-		FORCEINLINE bool operator()(const FPinTypeTreeInfoPtr A, const FPinTypeTreeInfoPtr B) const
-		{
-			return (A->GetDescription().ToString() < B->GetDescription().ToString());
-		}
-	};
 
 public:
 	static void FillLoadedTypesDatabase(FTypesDatabase& TypesDatabase, bool bIndexTypesOnly)
@@ -442,9 +429,8 @@ public:
 			AssetRegistryModule.Get().GetAssetsByClass(UUserDefinedEnum::StaticClass()->GetClassPathName(), AssetData);
 
 			FTypesDatabase::FUnLoadedTypesList UnLoadedTypesList = MakeShareable(new TArray<FUnloadedAssetData>());
-			for (int32 AssetIndex = 0; AssetIndex < AssetData.Num(); ++AssetIndex)
+			for (const FAssetData& Asset : AssetData)
 			{
-				const FAssetData& Asset = AssetData[AssetIndex];
 				if (Asset.IsValid() && !Asset.IsAssetLoaded())
 				{
 					UnLoadedTypesList->Add(FUnloadedAssetData(Asset));
@@ -461,9 +447,8 @@ public:
 				AssetRegistryModule.Get().GetAssetsByClass(UUserDefinedStruct::StaticClass()->GetClassPathName(), AssetData);
 
 				FTypesDatabase::FUnLoadedTypesList UnLoadedTypesList = MakeShareable(new TArray<FUnloadedAssetData>());
-				for (int32 AssetIndex = 0; AssetIndex < AssetData.Num(); ++AssetIndex)
+				for (const FAssetData& Asset : AssetData)
 				{
-					const FAssetData& Asset = AssetData[AssetIndex];
 					if (Asset.IsValid() && !Asset.IsAssetLoaded())
 					{
 						UnLoadedTypesList->Add(FUnloadedAssetData(Asset));
@@ -483,10 +468,8 @@ public:
 				FTypesDatabase::FUnLoadedTypesList UnLoadedInterfacesList = MakeShareable(new TArray<FUnloadedAssetData>());
 				FTypesDatabase::FUnLoadedTypesList UnLoadedClassesList = MakeShareable(new TArray<FUnloadedAssetData>());
 
-				for (int32 AssetIndex = 0; AssetIndex < AssetData.Num(); ++AssetIndex)
+				for (const FAssetData& Asset : AssetData)
 				{
-					const FAssetData& Asset = AssetData[AssetIndex];
-
 					if (Asset.IsValid() && !Asset.IsAssetLoaded())
 					{
 						const FString BlueprintTypeStr = Asset.GetTagValueRef<FString>(FBlueprintTags::BlueprintType);
@@ -527,40 +510,45 @@ public:
 	{
 		TRACE_CPUPROFILER_EVENT_SCOPE(WILD_FGatherTypesHelper::Gather);
 
-		FEdGraphPinType LoadedPinSubtype;
-		LoadedPinSubtype.PinCategory = (CategoryName == UEdGraphSchema_K2::PC_Enum ? UEdGraphSchema_K2::PC_Byte : CategoryName);
-		LoadedPinSubtype.PinSubCategory = NAME_None;
-		LoadedPinSubtype.PinSubCategoryObject = nullptr;
+		const FName PinCategory = (CategoryName == UEdGraphSchema_K2::PC_Enum ? UEdGraphSchema_K2::PC_Byte : CategoryName);
 
 		FTypesDatabase::FLoadedTypesList* LoadedSubTypesPtr = TypesDatabase.LoadedTypesMap.Find(CategoryName);
+		FTypesDatabase::FUnLoadedTypesList* UnLoadedSubTypesPtr = TypesDatabase.UnLoadedTypesMap.Find(CategoryName);
+		if (LoadedSubTypesPtr && UnLoadedSubTypesPtr && LoadedSubTypesPtr->IsValid() && UnLoadedSubTypesPtr->IsValid())
+		{
+			OutChildren.Reserve(LoadedSubTypesPtr->Get()->Num() + UnLoadedSubTypesPtr->Get()->Num());
+		}
+
 		if (LoadedSubTypesPtr && LoadedSubTypesPtr->IsValid())
 		{
 			for (FLoadedAssetData& LoadedAssetData : *LoadedSubTypesPtr->Get())
 			{
-				OutChildren.Add(MakeShareable(new UEdGraphSchema_K2::FPinTypeTreeInfo(LoadedPinSubtype.PinCategory
+				OutChildren.Add(MakeShared<UEdGraphSchema_K2::FPinTypeTreeInfo>(PinCategory
 					, LoadedAssetData.Object
 					, LoadedAssetData.Tooltip.IsEmpty() ? FriendlyName : LoadedAssetData.Tooltip
 					, false
-					, LoadedAssetData.PossibleObjectReferenceTypes)));
+					, LoadedAssetData.PossibleObjectReferenceTypes));
 			}
 		}
 
-		FTypesDatabase::FUnLoadedTypesList* UnLoadedSubTypesPtr = TypesDatabase.UnLoadedTypesMap.Find(CategoryName);
 		if (UnLoadedSubTypesPtr && UnLoadedSubTypesPtr->IsValid())
 		{
 			for (FUnloadedAssetData& It : *UnLoadedSubTypesPtr->Get())
 			{
-				FPinTypeTreeInfoPtr TypeTreeInfo = MakeShareable(new UEdGraphSchema_K2::FPinTypeTreeInfo(It.AssetFriendlyName
+				OutChildren.Add(MakeShared<UEdGraphSchema_K2::FPinTypeTreeInfo>(It.AssetFriendlyName
 					, CategoryName
 					, It.AssetData
 					, It.Tooltip
 					, false
 					, It.PossibleObjectReferenceTypes));
-				OutChildren.Add(TypeTreeInfo);
 			}
 		}
 
-		OutChildren.Sort(FCompareChildren());
+		OutChildren.Sort(
+			[](const FPinTypeTreeInfoPtr& A, const FPinTypeTreeInfoPtr& B)
+			{
+				return A->GetCachedDescriptionString().Compare(B->GetCachedDescriptionString()) < 0;
+			});
 	}
 };
 
@@ -588,13 +576,14 @@ const FEdGraphPinType& UEdGraphSchema_K2::FPinTypeTreeInfo::GetPinType(bool bFor
 	return PinType;
 }
 
-void UEdGraphSchema_K2::FPinTypeTreeInfo::Init(const FText& InFriendlyName, const FName CategoryName, const UEdGraphSchema_K2* Schema, const FText& InTooltip, bool bInReadOnly, FTypesDatabase* TypesDatabase)
+UEdGraphSchema_K2::FPinTypeTreeInfo::FPinTypeTreeInfo(const FText& InFriendlyName, const FName CategoryName, const UEdGraphSchema_K2* Schema, const FText& InTooltip, bool bInReadOnly/*=false*/, FTypesDatabase* TypesDatabase /*=nullptr*/)
+	: PossibleObjectReferenceTypes(0)
 {
 	TRACE_CPUPROFILER_EVENT_SCOPE(WILD_FPinTypeTreeInfo::Init);
 	check( !CategoryName.IsNone() );
 	check( Schema );
+	check(!InFriendlyName.IsEmpty());
 
-	FriendlyName = InFriendlyName;
 	Tooltip = InTooltip;
 	PinType.PinCategory = (CategoryName == PC_Enum ? PC_Byte : CategoryName);
 	PinType.PinSubCategory = (CategoryName == PC_Real ? PC_Double : NAME_None);
@@ -602,7 +591,8 @@ void UEdGraphSchema_K2::FPinTypeTreeInfo::Init(const FText& InFriendlyName, cons
 
 	bReadOnly = bInReadOnly;
 
-	CachedDescription = GenerateDescription();
+	CachedDescription = InFriendlyName;
+	CachedDescriptionString = MakeShared<FString>(CachedDescription.ToString());
 
 	if (Schema->DoesTypeHaveSubtypes(CategoryName))
 	{
@@ -625,15 +615,15 @@ UEdGraphSchema_K2::FPinTypeTreeInfo::FPinTypeTreeInfo(const FName CategoryName,
 
 	bReadOnly = bInReadOnly;
 	CachedDescription = GenerateDescription();
+	CachedDescriptionString = MakeShared<FString>(CachedDescription.ToString());
 }
 
 UEdGraphSchema_K2::FPinTypeTreeInfo::FPinTypeTreeInfo(const FText& InFriendlyName, const FName CategoryName, const FAssetData& AssetData, const FText& InTooltip, bool bInReadOnly, uint8 InPossibleObjectReferenceTypes)
 	: PossibleObjectReferenceTypes(InPossibleObjectReferenceTypes)
 {
-	FriendlyName = InFriendlyName;
-
 	check(!CategoryName.IsNone());
 	check(AssetData.IsValid());
+	check(!InFriendlyName.IsEmpty());
 
 	Tooltip = InTooltip;
 	PinType.PinCategory = CategoryName;
@@ -641,49 +631,40 @@ UEdGraphSchema_K2::FPinTypeTreeInfo::FPinTypeTreeInfo(const FText& InFriendlyNam
 	CachedAssetData = AssetData;
 
 	bReadOnly = bInReadOnly;
-	CachedDescription = GenerateDescription();
+	CachedDescription = InFriendlyName;
+	CachedDescriptionString = MakeShared<FString>(CachedDescription.ToString());
 }
 
 UEdGraphSchema_K2::FPinTypeTreeInfo::FPinTypeTreeInfo(TSharedPtr<FPinTypeTreeInfo> InInfo)
 {
 	PinType = InInfo->PinType;
 	bReadOnly = InInfo->bReadOnly;
-	FriendlyName = InInfo->FriendlyName;
 	CachedAssetData = InInfo->CachedAssetData;
 	Tooltip = InInfo->Tooltip;
 	CachedDescription = InInfo->CachedDescription;
+	CachedDescriptionString = InInfo->CachedDescriptionString;
 	PossibleObjectReferenceTypes = InInfo->PossibleObjectReferenceTypes;
 }
 
+const FText& UEdGraphSchema_K2::FPinTypeTreeInfo::GetDescription() const
+{
+	return CachedDescription;
+}
+
 FText UEdGraphSchema_K2::FPinTypeTreeInfo::GenerateDescription()
 {
-	if (!FriendlyName.IsEmpty())
-	{
-		return FriendlyName;
-	}
-	else if (PinType.PinSubCategoryObject.IsValid())
+	check(PinType.PinSubCategoryObject.IsValid());
+	FText DisplayName;
+	if (UField* SubCategoryField = Cast<UField>(PinType.PinSubCategoryObject.Get()))
 	{
-		FText DisplayName;
-		if (UField* SubCategoryField = Cast<UField>(PinType.PinSubCategoryObject.Get()))
-		{
-			DisplayName = SubCategoryField->GetDisplayNameText();
-		}
-		else
-		{
-			DisplayName = FText::FromString(FName::NameToDisplayString(PinType.PinSubCategoryObject->GetName(), PinType.PinCategory == PC_Boolean));
-		}
-
-		return DisplayName;
+		DisplayName = SubCategoryField->GetDisplayNameText();
 	}
 	else
 	{
-		return LOCTEXT("PinDescriptionError", "Error!");
+		DisplayName = FText::FromString(FName::NameToDisplayString(PinType.PinSubCategoryObject->GetName(), PinType.PinCategory == PC_Boolean));
 	}
-}
 
-FText UEdGraphSchema_K2::FPinTypeTreeInfo::GetDescription() const
-{
-	return CachedDescription;
+	return DisplayName;
 }
 
 const FAssetData& UEdGraphSchema_K2::FPinTypeTreeInfo::GetCachedAssetData() const
diff --git a/Engine/Source/Editor/KismetWidgets/Private/SPinTypeSelector.cpp b/Engine/Source/Editor/KismetWidgets/Private/SPinTypeSelector.cpp
--- a/Engine/Source/Editor/KismetWidgets/Private/SPinTypeSelector.cpp
+++ b/Engine/Source/Editor/KismetWidgets/Private/SPinTypeSelector.cpp
@@ -1433,7 +1433,7 @@ bool SPinTypeSelector::GetChildrenWithSupportedTypes(const TArray<FPinTypeTreeIt
 	for( auto it = UnfilteredList.CreateConstIterator(); it; ++it )
 	{
 		FPinTypeTreeItem Item = *it;
-		FPinTypeTreeItem NewInfo = MakeShareable( new UEdGraphSchema_K2::FPinTypeTreeInfo(Item) );
+		FPinTypeTreeItem NewInfo = MakeShared<UEdGraphSchema_K2::FPinTypeTreeInfo>(Item);
 		TArray<FPinTypeTreeItem> ValidChildren;
 
 		const bool bHasChildrenWithValidTypes = GetChildrenWithSupportedTypes(Item->Children, ValidChildren);
@@ -1490,12 +1490,10 @@ bool SPinTypeSelector::GetChildrenMatchingSearch(const FText& InSearchText, cons
 
 	bool bReturnVal = false;
 
-	for( auto it = UnfilteredList.CreateConstIterator(); it; ++it )
+	for( const FPinTypeTreeItem& Item : UnfilteredList)
 	{
-		FPinTypeTreeItem Item = *it;
-		FPinTypeTreeItem NewInfo = MakeShareable( new UEdGraphSchema_K2::FPinTypeTreeInfo(Item) );
 		TArray<FPinTypeTreeItem> ValidChildren;
-
+		ValidChildren.Reserve(Item->Children.Num());
 		const bool bHasChildrenMatchingSearch = GetChildrenMatchingSearch(InSearchText, Item->Children, ValidChildren, OutTopLevenshteinResult);
 		bool bFilterMatches = bIsEmptySearch;
 
@@ -1519,9 +1517,8 @@ bool SPinTypeSelector::GetChildrenMatchingSearch(const FText& InSearchText, cons
 			// If we didn't match the custom filter, or it's an empty search, let's not do any checks against the FilterTerms
 			if (bFilterMatches && !bIsEmptySearch)
 			{
-				const FText LocalizedDescription = Item->GetDescription();
-				const FString LocalizedDescriptionString = LocalizedDescription.ToString();
-				const FString* SourceDescriptionStringPtr = FTextInspector::GetSourceString(LocalizedDescription);
+				const FString& LocalizedDescriptionString = Item->GetCachedDescriptionString();
+				const FString* SourceDescriptionStringPtr = FTextInspector::GetSourceString(Item->GetDescription());
 
 				// Test both the localized and source strings for a match
 				const FString MangledLocalizedDescriptionString = LocalizedDescriptionString.Replace(TEXT(" "), TEXT(""));
@@ -1538,13 +1535,14 @@ bool SPinTypeSelector::GetChildrenMatchingSearch(const FText& InSearchText, cons
 		if( bHasChildrenMatchingSearch
 			|| bFilterMatches )
 		{
+			FPinTypeTreeItem NewInfo = MakeShared<UEdGraphSchema_K2::FPinTypeTreeInfo>(Item);
 			NewInfo->Children = ValidChildren;
 			OutFilteredList.Add(NewInfo);
 
 			// Skip anything with children, those are categories and we don't care about those.
 			if (Item->Children.Num() == 0)
 			{
-				OutTopLevenshteinResult.CompareAndUpdate(TrimmedFilterString, NewInfo, Item->GetDescription().ToString());
+				OutTopLevenshteinResult.CompareAndUpdate(TrimmedFilterString, NewInfo, Item->GetCachedDescriptionString());
 			}
 
 			if (TypeTreeView.IsValid())
