Combine FUnloadedAssetData/FLoadedAssetData and various FPinTypeTreeInfo constructors, remove the intermediate 'type database', filter TypeTreeRoot in place, perform 'custom filtering' during filter step, rather than during 'children matching' search

#changelist validated
#virtualized

diff --git a/Engine/Source/Editor/BlueprintGraph/Classes/EdGraphSchema_K2.h b/Engine/Source/Editor/BlueprintGraph/Classes/EdGraphSchema_K2.h
--- a/Engine/Source/Editor/BlueprintGraph/Classes/EdGraphSchema_K2.h
+++ b/Engine/Source/Editor/BlueprintGraph/Classes/EdGraphSchema_K2.h
@@ -45,7 +45,6 @@ class UScriptStruct;
 class UToolMenu;
 struct FAssetData;
 struct FToolMenuSection;
-struct FTypesDatabase;
 template <typename T> struct TObjectPtr;
 
 /** Reference to an structure (only used in 'docked' palette) */
@@ -346,8 +345,6 @@ enum class ETypeTreeFilter : uint8
 
 ENUM_CLASS_FLAGS(ETypeTreeFilter);
 
-struct FTypesDatabase;
-
 UCLASS(config=Editor)
 class BLUEPRINTGRAPH_API UEdGraphSchema_K2 : public UEdGraphSchema
 {
@@ -473,7 +470,7 @@ public:
 			PinType.PinSubCategory = SubCategory;
 		}
 
-		FPinTypeTreeInfo(const FText& InFriendlyName, const FName CategoryName, const UEdGraphSchema_K2* Schema, const FText& InTooltip, bool bInReadOnly = false, FTypesDatabase* TypesDatabase = nullptr);
+		FPinTypeTreeInfo(const FText& InFriendlyName, const FName CategoryName, const UEdGraphSchema_K2* Schema, const FText& InTooltip, bool bInReadOnly = false);
 		FPinTypeTreeInfo(const FName CategoryName, UObject* SubCategoryObject, const FText& InTooltip, bool bInReadOnly = false, uint8 InPossibleObjectReferenceTypes = 0);
 		FPinTypeTreeInfo(const FText& InFriendlyName, const FName CategoryName, const FAssetData& AssetData, const FText& InTooltip, bool bInReadOnly = false, uint8 InPossibleObjectReferenceTypes = 0);
 		FPinTypeTreeInfo(TSharedPtr<FPinTypeTreeInfo> InInfo);
diff --git a/Engine/Source/Editor/BlueprintGraph/Private/EdGraphSchema_K2.cpp b/Engine/Source/Editor/BlueprintGraph/Private/EdGraphSchema_K2.cpp
--- a/Engine/Source/Editor/BlueprintGraph/Private/EdGraphSchema_K2.cpp
+++ b/Engine/Source/Editor/BlueprintGraph/Private/EdGraphSchema_K2.cpp
@@ -287,270 +287,202 @@ const FName FBlueprintMetadata::MD_FieldNotify(TEXT("FieldNotify"));
 
 #define LOCTEXT_NAMESPACE "KismetSchema"
 
-struct FUnloadedAssetData
+/** Helpers for gathering pin type tree info for enums, structs, classes, and interfaces */
+namespace GatherPinsImpl
 {
-	FAssetData AssetData;
-	FText AssetFriendlyName;
-	FText Tooltip;
-	uint8 PossibleObjectReferenceTypes;
+	TSharedPtr<UEdGraphSchema_K2::FPinTypeTreeInfo> FromAssetData(const FAssetData& InAsset, FName CategoryName, EObjectReferenceType ReferenceType);
+	TSharedPtr<UEdGraphSchema_K2::FPinTypeTreeInfo> FromObject(UField* Field, FName CategoryName, EObjectReferenceType ReferenceType);
+	void SortPinTypes(TArray<TSharedPtr<UEdGraphSchema_K2::FPinTypeTreeInfo>>& PinArray);
 
-	FUnloadedAssetData()
-		: PossibleObjectReferenceTypes(0)
-	{}
+	void FindEnums(const TSharedPtr<UEdGraphSchema_K2::FPinTypeTreeInfo>& Owner);
+	void FindStructs(const TSharedPtr<UEdGraphSchema_K2::FPinTypeTreeInfo>& Owner);
+	void FindObjectsAndInterfaces(const TSharedPtr<UEdGraphSchema_K2::FPinTypeTreeInfo>& ObjectsOwner, const TSharedPtr<UEdGraphSchema_K2::FPinTypeTreeInfo>& InterfacesOwner);
+}
 
-	FUnloadedAssetData(const FAssetData& InAsset, uint8 InPossibleObjectReferenceTypes = 0)
-		: AssetData(InAsset)
-		, AssetFriendlyName(FText::FromString(FName::NameToDisplayString(InAsset.AssetName.ToString(), false)))
-		, PossibleObjectReferenceTypes(InPossibleObjectReferenceTypes)
+TSharedPtr<UEdGraphSchema_K2::FPinTypeTreeInfo> GatherPinsImpl::FromAssetData(const FAssetData& InAsset, FName CategoryName, EObjectReferenceType ReferenceType)
+{
+	return MakeShared<UEdGraphSchema_K2::FPinTypeTreeInfo>(
+		FText::FromString(FName::NameToDisplayString(InAsset.AssetName.ToString(), false))
+		, CategoryName
+		, InAsset
+		, FText::FromString(InAsset.GetObjectPathString())
+		, false
+		, (uint8)ReferenceType);
+}
+
+TSharedPtr<UEdGraphSchema_K2::FPinTypeTreeInfo> GatherPinsImpl::FromObject(UField* Field, FName CategoryName, EObjectReferenceType ReferenceType)
+{
+	return MakeShared<UEdGraphSchema_K2::FPinTypeTreeInfo>(
+		CategoryName
+		, Field
+		, Field->GetToolTipText()
+		, false
+		, (uint8)ReferenceType);
+}
+
+void GatherPinsImpl::SortPinTypes(TArray<TSharedPtr<UEdGraphSchema_K2::FPinTypeTreeInfo>>& PinArray)
+{
+	PinArray.Sort(
+		[](const TSharedPtr<UEdGraphSchema_K2::FPinTypeTreeInfo>& A, const TSharedPtr<UEdGraphSchema_K2::FPinTypeTreeInfo>& B)
+		{
+			return A->GetCachedDescriptionString().Compare(B->GetCachedDescriptionString()) < 0;
+		});
+}
+
+void GatherPinsImpl::FindEnums(const TSharedPtr<UEdGraphSchema_K2::FPinTypeTreeInfo>& Owner)
+{
+	TSet<FTopLevelAssetPath> ProcessedAssets;
+
+	check(Owner->bReadOnly);
+	// Generate a list of all potential enums which have "BlueprintType=true" in their metadata
+	for (TObjectIterator<UEnum> EnumIt; EnumIt; ++EnumIt)
 	{
-		InAsset.GetTagValue(FBlueprintMetadata::MD_Tooltip, Tooltip);
-		if (Tooltip.IsEmpty())
+		UEnum* CurrentEnum = *EnumIt;
+		ProcessedAssets.Add(FTopLevelAssetPath(CurrentEnum));
+		if (UEdGraphSchema_K2::IsAllowableBlueprintVariableType(CurrentEnum))
 		{
-			Tooltip = FText::FromString(InAsset.GetObjectPathString());
+			Owner->Children.Emplace(
+				FromObject(CurrentEnum
+					, UEdGraphSchema_K2::PC_Enum
+					, EObjectReferenceType::NotAnObject));
 		}
 	}
-};
-
-struct FLoadedAssetData
-{
-	FText Tooltip;
-	UObject* Object;
-	uint8 PossibleObjectReferenceTypes;
 
-	FLoadedAssetData() 
-		: Object(nullptr)
-		, PossibleObjectReferenceTypes(0) {}
+	TArray<FAssetData> AssetData;
+	const FAssetRegistryModule& AssetRegistryModule = FModuleManager::LoadModuleChecked<FAssetRegistryModule>(TEXT("AssetRegistry"));
+	AssetRegistryModule.Get().GetAssetsByClass(UUserDefinedEnum::StaticClass()->GetClassPathName(), AssetData);
 
-	FLoadedAssetData(UObject* InObject, uint8 InPossibleObjectReferenceTypes = 0)
-		: Object(InObject) 
-		, PossibleObjectReferenceTypes(InPossibleObjectReferenceTypes)
+	for (const FAssetData& Asset : AssetData)
 	{
-		UStruct* Struct = Cast<UStruct>(Object);
-		Tooltip = Struct ? Struct->GetToolTipText() : FText::GetEmpty();
+		if (Asset.IsValid() && !ProcessedAssets.Contains(FTopLevelAssetPath(Asset.PackageName, Asset.AssetName)))
+		{
+			Owner->Children.Emplace(
+				FromAssetData(Asset
+					, UEdGraphSchema_K2::PC_Enum
+					, EObjectReferenceType::NotAnObject));
+		}
 	}
-};
-
-struct FTypesDatabase
-{
-	typedef TSharedPtr<TArray<FLoadedAssetData>> FLoadedTypesList;
-	TMap<FName, FLoadedTypesList> LoadedTypesMap;
 
-	typedef TSharedPtr<TArray<FUnloadedAssetData>> FUnLoadedTypesList;
-	TMap<FName, FUnLoadedTypesList> UnLoadedTypesMap;
-};
+	SortPinTypes(Owner->Children);
+}
 
-/** Helper class to gather variable types */
-class FGatherTypesHelper
+void GatherPinsImpl::FindStructs(const TSharedPtr<UEdGraphSchema_K2::FPinTypeTreeInfo>& Owner)
 {
-private:
-	typedef TSharedPtr<UEdGraphSchema_K2::FPinTypeTreeInfo> FPinTypeTreeInfoPtr;
+	check(Owner->bReadOnly);
+	TSet<FTopLevelAssetPath> ProcessedAssets;
 
-public:
-	static void FillLoadedTypesDatabase(FTypesDatabase& TypesDatabase, bool bIndexTypesOnly)
+	// Find script structs marked with "BlueprintType=true" in their metadata, and add to the list
+	for (TObjectIterator<UScriptStruct> StructIt; StructIt; ++StructIt)
 	{
-		TRACE_CPUPROFILER_EVENT_SCOPE(WILD_FGatherTypesHelper::FillLoadedTypesDatabase);
+		UScriptStruct* ScriptStruct = *StructIt;
+		ProcessedAssets.Add(FTopLevelAssetPath(ScriptStruct));
+		if (UEdGraphSchema_K2::IsAllowableBlueprintVariableType(ScriptStruct))
+		{
+			Owner->Children.Emplace(
+				FromObject(ScriptStruct
+					, UEdGraphSchema_K2::PC_Struct
+					, EObjectReferenceType::NotAnObject));
+		}
+	}
 
-		// Loaded types
-		TypesDatabase.LoadedTypesMap.Reset();
+	TArray<FAssetData> AssetData;
+	const FAssetRegistryModule& AssetRegistryModule = FModuleManager::LoadModuleChecked<FAssetRegistryModule>(TEXT("AssetRegistry"));
+	AssetRegistryModule.Get().GetAssetsByClass(UUserDefinedStruct::StaticClass()->GetClassPathName(), AssetData);
 
-		//(Type == UEdGraphSchema_K2::PC_Enum)
+	for (const FAssetData& Asset : AssetData)
+	{
+		if (Asset.IsValid() && !ProcessedAssets.Contains(FTopLevelAssetPath(Asset.PackageName, Asset.AssetName)))
 		{
-			FTypesDatabase::FLoadedTypesList LoadedTypesList = MakeShareable(new TArray<FLoadedAssetData>());
-			// Generate a list of all potential enums which have "BlueprintType=true" in their metadata
-			for (TObjectIterator<UEnum> EnumIt; EnumIt; ++EnumIt)
-			{
-				UEnum* CurrentEnum = *EnumIt;
-				if (UEdGraphSchema_K2::IsAllowableBlueprintVariableType(CurrentEnum))
-				{
-					LoadedTypesList->Add(FLoadedAssetData(CurrentEnum));
-				}
-			}
-			TypesDatabase.LoadedTypesMap.Add(UEdGraphSchema_K2::PC_Enum, LoadedTypesList);
+			Owner->Children.Emplace(
+				FromAssetData(Asset
+					, UEdGraphSchema_K2::PC_Struct
+					, EObjectReferenceType::NotAnObject));
 		}
+	}
+
+	SortPinTypes(Owner->Children);
+}
+
+void GatherPinsImpl::FindObjectsAndInterfaces(const TSharedPtr<UEdGraphSchema_K2::FPinTypeTreeInfo>& ObjectsOwner, const TSharedPtr<UEdGraphSchema_K2::FPinTypeTreeInfo>& InterfacesOwner)
+{
+	check(ObjectsOwner->bReadOnly && InterfacesOwner->bReadOnly);
+	TSet<FTopLevelAssetPath> ProcessedAssets;
 
-		if (!bIndexTypesOnly)
+	// Generate a list of all potential objects which have "BlueprintType=true" in their metadata
+	for (TObjectIterator<UClass> ClassIt; ClassIt; ++ClassIt)
+	{
+		UClass* CurrentClass = *ClassIt;
+		ProcessedAssets.Add(FTopLevelAssetPath(CurrentClass));
+		ProcessedAssets.Add(FTopLevelAssetPath(CurrentClass->ClassGeneratedBy));
+		const bool bIsInterface = CurrentClass->IsChildOf(UInterface::StaticClass());
+		const bool bIsBlueprintType = UEdGraphSchema_K2::IsAllowableBlueprintVariableType(CurrentClass);
+		const bool bIsDeprecated = CurrentClass->HasAnyClassFlags(CLASS_Deprecated);
+		if (bIsBlueprintType && !bIsDeprecated)
 		{
-			//(Type == UEdGraphSchema_K2::PC_Struct)
+			if (bIsInterface)
 			{
-				FTypesDatabase::FLoadedTypesList LoadedTypesList = MakeShareable(new TArray<FLoadedAssetData>());
-				// Find script structs marked with "BlueprintType=true" in their metadata, and add to the list
-				for (TObjectIterator<UScriptStruct> StructIt; StructIt; ++StructIt)
-				{
-					UScriptStruct* ScriptStruct = *StructIt;
-					if (UEdGraphSchema_K2::IsAllowableBlueprintVariableType(ScriptStruct))
-					{
-						LoadedTypesList->Add(FLoadedAssetData(ScriptStruct));
-					}
-				}
-				TypesDatabase.LoadedTypesMap.Add(UEdGraphSchema_K2::PC_Struct, LoadedTypesList);
+				InterfacesOwner->Children.Emplace(
+					FromObject(CurrentClass
+						, UEdGraphSchema_K2::PC_Interface
+						, EObjectReferenceType::NotAnObject));
 			}
-
-			//(Type == UEdGraphSchema_K2::PC_Class || Type == UEdGraphSchema_K2::PC_SoftClass) UEdGraphSchema_K2::PC_Interface)
-			//(Type == UEdGraphSchema_K2::PC_Object || Type == UEdGraphSchema_K2::PC_SoftObject)
+			else
 			{
-				FTypesDatabase::FLoadedTypesList InterfaceLoadedTypesList = MakeShareable(new TArray<FLoadedAssetData>());
-				FTypesDatabase::FLoadedTypesList AllObjectLoadedTypesList = MakeShareable(new TArray<FLoadedAssetData>());
-
-				// Generate a list of all potential objects which have "BlueprintType=true" in their metadata
-				for (TObjectIterator<UClass> ClassIt; ClassIt; ++ClassIt)
-				{
-					UClass* CurrentClass = *ClassIt;
-					const bool bIsInterface = CurrentClass->IsChildOf(UInterface::StaticClass());
-					const bool bIsBlueprintType = UEdGraphSchema_K2::IsAllowableBlueprintVariableType(CurrentClass);
-					const bool bIsDeprecated = CurrentClass->HasAnyClassFlags(CLASS_Deprecated);
-					if (bIsBlueprintType && !bIsDeprecated)
-					{
-						if (bIsInterface)
-						{
-							InterfaceLoadedTypesList->Add(FLoadedAssetData(CurrentClass));
-						}
-						else
-						{
-							AllObjectLoadedTypesList->Add(FLoadedAssetData(CurrentClass, static_cast<uint8>(EObjectReferenceType::AllTypes)));
-						}
-					}
-				}
-				TypesDatabase.LoadedTypesMap.Add(UEdGraphSchema_K2::AllObjectTypes, AllObjectLoadedTypesList);
-				TypesDatabase.LoadedTypesMap.Add(UEdGraphSchema_K2::PC_Interface, InterfaceLoadedTypesList);
+				ObjectsOwner->Children.Emplace(
+					FromObject(CurrentClass
+						, UEdGraphSchema_K2::AllObjectTypes
+						, EObjectReferenceType::AllTypes));
 			}
 		}
 	}
 
-	static void FillUnLoadedTypesDatabase(FTypesDatabase& TypesDatabase, bool bIndexTypesOnly)
-	{
-		TRACE_CPUPROFILER_EVENT_SCOPE(WILD_FGatherTypesHelper::FillUnLoadedTypesDatabase);
-		// Loaded types
-		TypesDatabase.UnLoadedTypesMap.Reset();
+	TArray<FAssetData> AssetData;
+	const FAssetRegistryModule& AssetRegistryModule = FModuleManager::LoadModuleChecked<FAssetRegistryModule>(TEXT("AssetRegistry"));
+	AssetRegistryModule.Get().GetAssetsByClass(UBlueprint::StaticClass()->GetClassPathName(), AssetData);
 
-		const FAssetRegistryModule& AssetRegistryModule = FModuleManager::LoadModuleChecked<FAssetRegistryModule>(TEXT("AssetRegistry"));
+	const FStringView BPInterfaceTypeAllowed(TEXT("BPTYPE_Interface"));
+	const FStringView BPNormalTypeAllowed(TEXT("BPTYPE_Normal"));
 
+	for (const FAssetData& Asset : AssetData)
+	{
+		if (Asset.IsValid() && !ProcessedAssets.Contains(FTopLevelAssetPath(Asset.PackageName, Asset.AssetName)))
 		{
-			TArray<FAssetData> AssetData;
-			AssetRegistryModule.Get().GetAssetsByClass(UUserDefinedEnum::StaticClass()->GetClassPathName(), AssetData);
-
-			FTypesDatabase::FUnLoadedTypesList UnLoadedTypesList = MakeShareable(new TArray<FUnloadedAssetData>());
-			for (const FAssetData& Asset : AssetData)
+			FAssetDataTagMapSharedView::FFindTagResult FoundValue = Asset.TagsAndValues.FindTag(FBlueprintTags::BlueprintType);
+			if (!FoundValue.IsSet())
 			{
-				if (Asset.IsValid() && !Asset.IsAssetLoaded())
-				{
-					UnLoadedTypesList->Add(FUnloadedAssetData(Asset));
-				}
+				continue;
 			}
 
-			TypesDatabase.UnLoadedTypesMap.Add(UEdGraphSchema_K2::PC_Enum, UnLoadedTypesList);
-		}
+			const bool bNormalBP = FoundValue.Equals(BPNormalTypeAllowed);
+			const bool bInterfaceBP = FoundValue.Equals(BPInterfaceTypeAllowed);
 
-		if (!bIndexTypesOnly)
-		{
+			if (bNormalBP || bInterfaceBP)
 			{
-				TArray<FAssetData> AssetData;
-				AssetRegistryModule.Get().GetAssetsByClass(UUserDefinedStruct::StaticClass()->GetClassPathName(), AssetData);
-
-				FTypesDatabase::FUnLoadedTypesList UnLoadedTypesList = MakeShareable(new TArray<FUnloadedAssetData>());
-				for (const FAssetData& Asset : AssetData)
+				const uint32 ClassFlags = Asset.GetTagValueRef<uint32>(FBlueprintTags::ClassFlags);
+				if (!(ClassFlags & CLASS_Deprecated))
 				{
-					if (Asset.IsValid() && !Asset.IsAssetLoaded())
+					if (bNormalBP)
 					{
-						UnLoadedTypesList->Add(FUnloadedAssetData(Asset));
+						ObjectsOwner->Children.Emplace(
+							FromAssetData(Asset
+								, UEdGraphSchema_K2::AllObjectTypes
+								, EObjectReferenceType::AllTypes));
 					}
-				}
-
-				TypesDatabase.UnLoadedTypesMap.Add(UEdGraphSchema_K2::PC_Struct, UnLoadedTypesList);
-			}
-
-			{
-				TArray<FAssetData> AssetData;
-				AssetRegistryModule.Get().GetAssetsByClass(UBlueprint::StaticClass()->GetClassPathName(), AssetData);
-
-				const FString BPInterfaceTypeAllowed(TEXT("BPTYPE_Interface"));
-				const FString BPNormalTypeAllowed(TEXT("BPTYPE_Normal"));
-
-				FTypesDatabase::FUnLoadedTypesList UnLoadedInterfacesList = MakeShareable(new TArray<FUnloadedAssetData>());
-				FTypesDatabase::FUnLoadedTypesList UnLoadedClassesList = MakeShareable(new TArray<FUnloadedAssetData>());
-
-				for (const FAssetData& Asset : AssetData)
-				{
-					if (Asset.IsValid() && !Asset.IsAssetLoaded())
+					else if (bInterfaceBP)
 					{
-						const FString BlueprintTypeStr = Asset.GetTagValueRef<FString>(FBlueprintTags::BlueprintType);
-						const bool bNormalBP = BlueprintTypeStr == BPNormalTypeAllowed;
-						const bool bInterfaceBP = BlueprintTypeStr == BPInterfaceTypeAllowed;
-
-						if (bNormalBP || bInterfaceBP)
-						{
-							const uint32 ClassFlags = Asset.GetTagValueRef<uint32>(FBlueprintTags::ClassFlags);
-							if (!(ClassFlags & CLASS_Deprecated))
-							{
-								if (bNormalBP)
-								{
-									UnLoadedClassesList->Add(FUnloadedAssetData(Asset, static_cast<uint8>(EObjectReferenceType::AllTypes)));
-								}
-								else if (bInterfaceBP)
-								{
-									UnLoadedInterfacesList->Add(FUnloadedAssetData(Asset));
-								}
-							}
-						}
+						InterfacesOwner->Children.Emplace(
+							FromAssetData(Asset
+								, UEdGraphSchema_K2::PC_Interface
+								, EObjectReferenceType::NotAnObject));
 					}
 				}
-				TypesDatabase.UnLoadedTypesMap.Add(UEdGraphSchema_K2::PC_Interface, UnLoadedInterfacesList);
-				TypesDatabase.UnLoadedTypesMap.Add(UEdGraphSchema_K2::AllObjectTypes, UnLoadedClassesList);
 			}
 		}
 	}
 
-	/**
-	 * Gathers all valid sub-types (loaded and unloaded) of a passed category and sorts them alphabetically
-	 * @param FriendlyName		Friendly name to be used for the tooltip if there is no available data
-	 * @param CategoryName		Category (type) to find sub-types of
-	 * @param TypesDatabase		Types database
-	 * @param OutChildren		All the gathered children
-	 */
-	static void Gather(const FText& FriendlyName, const FName CategoryName, FTypesDatabase& TypesDatabase, TArray<FPinTypeTreeInfoPtr>& OutChildren)
-	{
-		TRACE_CPUPROFILER_EVENT_SCOPE(WILD_FGatherTypesHelper::Gather);
-
-		const FName PinCategory = (CategoryName == UEdGraphSchema_K2::PC_Enum ? UEdGraphSchema_K2::PC_Byte : CategoryName);
-
-		FTypesDatabase::FLoadedTypesList* LoadedSubTypesPtr = TypesDatabase.LoadedTypesMap.Find(CategoryName);
-		FTypesDatabase::FUnLoadedTypesList* UnLoadedSubTypesPtr = TypesDatabase.UnLoadedTypesMap.Find(CategoryName);
-		if (LoadedSubTypesPtr && UnLoadedSubTypesPtr && LoadedSubTypesPtr->IsValid() && UnLoadedSubTypesPtr->IsValid())
-		{
-			OutChildren.Reserve(LoadedSubTypesPtr->Get()->Num() + UnLoadedSubTypesPtr->Get()->Num());
-		}
-
-		if (LoadedSubTypesPtr && LoadedSubTypesPtr->IsValid())
-		{
-			for (FLoadedAssetData& LoadedAssetData : *LoadedSubTypesPtr->Get())
-			{
-				OutChildren.Add(MakeShared<UEdGraphSchema_K2::FPinTypeTreeInfo>(PinCategory
-					, LoadedAssetData.Object
-					, LoadedAssetData.Tooltip.IsEmpty() ? FriendlyName : LoadedAssetData.Tooltip
-					, false
-					, LoadedAssetData.PossibleObjectReferenceTypes));
-			}
-		}
-
-		if (UnLoadedSubTypesPtr && UnLoadedSubTypesPtr->IsValid())
-		{
-			for (FUnloadedAssetData& It : *UnLoadedSubTypesPtr->Get())
-			{
-				OutChildren.Add(MakeShared<UEdGraphSchema_K2::FPinTypeTreeInfo>(It.AssetFriendlyName
-					, CategoryName
-					, It.AssetData
-					, It.Tooltip
-					, false
-					, It.PossibleObjectReferenceTypes));
-			}
-		}
-
-		OutChildren.Sort(
-			[](const FPinTypeTreeInfoPtr& A, const FPinTypeTreeInfoPtr& B)
-			{
-				return A->GetCachedDescriptionString().Compare(B->GetCachedDescriptionString()) < 0;
-			});
-	}
-};
+	SortPinTypes(InterfacesOwner->Children);
+	SortPinTypes(ObjectsOwner->Children);
+}
 
 const FEdGraphPinType& UEdGraphSchema_K2::FPinTypeTreeInfo::GetPinType(bool bForceLoadedSubCategoryObject)
 {
@@ -576,7 +508,7 @@ const FEdGraphPinType& UEdGraphSchema_K2::FPinTypeTreeInfo::GetPinType(bool bFor
 	return PinType;
 }
 
-UEdGraphSchema_K2::FPinTypeTreeInfo::FPinTypeTreeInfo(const FText& InFriendlyName, const FName CategoryName, const UEdGraphSchema_K2* Schema, const FText& InTooltip, bool bInReadOnly/*=false*/, FTypesDatabase* TypesDatabase /*=nullptr*/)
+UEdGraphSchema_K2::FPinTypeTreeInfo::FPinTypeTreeInfo(const FText& InFriendlyName, const FName CategoryName, const UEdGraphSchema_K2* Schema, const FText& InTooltip, bool bInReadOnly/*=false*/)
 	: PossibleObjectReferenceTypes(0)
 {
 	TRACE_CPUPROFILER_EVENT_SCOPE(WILD_FPinTypeTreeInfo::Init);
@@ -593,14 +525,6 @@ UEdGraphSchema_K2::FPinTypeTreeInfo::FPinTypeTreeInfo(const FText& InFriendlyNam
 
 	CachedDescription = InFriendlyName;
 	CachedDescriptionString = MakeShared<FString>(CachedDescription.ToString());
-
-	if (Schema->DoesTypeHaveSubtypes(CategoryName))
-	{
-		if (TypesDatabase)
-		{
-			FGatherTypesHelper::Gather(InFriendlyName, CategoryName, *TypesDatabase, Children);
-		}
-	}
 }
 
 UEdGraphSchema_K2::FPinTypeTreeInfo::FPinTypeTreeInfo(const FName CategoryName, UObject* SubCategoryObject, const FText& InTooltip, bool bInReadOnly/*=false*/, uint8 InPossibleObjectReferenceTypes)
@@ -4065,36 +3989,6 @@ void UEdGraphSchema_K2::GetVariableTypeTree(TArray< TSharedPtr<FPinTypeTreeInfo>
 	bool bIndexTypesOnly = (TypeTreeFilter & ETypeTreeFilter::IndexTypesOnly) == ETypeTreeFilter::IndexTypesOnly;
 	bool bRootTypesOnly = (TypeTreeFilter & ETypeTreeFilter::RootTypesOnly) == ETypeTreeFilter::RootTypesOnly;
 
-#ifdef SCHEMA_K2_GETVARIABLETYPETREE_LOG_TIME
-	static_assert(false, "Macro redefinition.");
-#endif
-#define SCHEMA_K2_GETVARIABLETYPETREE_LOG_TIME 0
-
-#if SCHEMA_K2_GETVARIABLETYPETREE_LOG_TIME
-	const double StartTime = FPlatformTime::Seconds();
-#endif //SCHEMA_K2_GETVARIABLETYPETREE_LOG_TIME
-
-	FTypesDatabase TypesDatabase;
-	FTypesDatabase* TypesDatabasePtr = nullptr;
-	if (!bRootTypesOnly)
-	{
-		TypesDatabasePtr = &TypesDatabase;
-		FGatherTypesHelper::FillLoadedTypesDatabase(TypesDatabase, bIndexTypesOnly);
-	}
-
-#if SCHEMA_K2_GETVARIABLETYPETREE_LOG_TIME
-	const double DatabaseLoadedTime = FPlatformTime::Seconds();
-#endif //SCHEMA_K2_GETVARIABLETYPETREE_LOG_TIME
-
-	if (!bRootTypesOnly)
-	{
-		FGatherTypesHelper::FillUnLoadedTypesDatabase(TypesDatabase, bIndexTypesOnly);
-	}
-
-#if SCHEMA_K2_GETVARIABLETYPETREE_LOG_TIME
-	const double DatabaseUnLoadedTime = FPlatformTime::Seconds();
-#endif //SCHEMA_K2_GETVARIABLETYPETREE_LOG_TIME
-
 	// Clear the list
 	TypeTree.Empty();
 
@@ -4132,28 +4026,36 @@ void UEdGraphSchema_K2::GetVariableTypeTree(TArray< TSharedPtr<FPinTypeTreeInfo>
 	// Add the types that have subtrees
 	if (!bIndexTypesOnly)
 	{
-		TypeTree.Add(MakeShareable(new FPinTypeTreeInfo(GetCategoryText(PC_Struct, true), PC_Struct, this, LOCTEXT("StructType", "Struct (value) types"), true, TypesDatabasePtr)));
-		TypeTree.Add(MakeShareable(new FPinTypeTreeInfo(GetCategoryText(PC_Interface, true), PC_Interface, this, LOCTEXT("InterfaceType", "Interface types"), true, TypesDatabasePtr)));
+		TSharedPtr<FPinTypeTreeInfo> Structs = MakeShared<FPinTypeTreeInfo>(GetCategoryText(PC_Struct, true), PC_Struct, this, LOCTEXT("StructType", "Struct (value) types"), true);
+		if (!bRootTypesOnly)
+		{
+			GatherPinsImpl::FindStructs(Structs);
+		}
+		TypeTree.Add(Structs);
+
+		TSharedPtr<FPinTypeTreeInfo> Interfaces = MakeShared<FPinTypeTreeInfo>(GetCategoryText(PC_Interface, true), PC_Interface, this, LOCTEXT("InterfaceType", "Interface types"), true);
+		TypeTree.Add(Interfaces);
 
 		if (!bRootTypesOnly)
 		{
-			TypeTree.Add(MakeShareable(new FPinTypeTreeInfo(GetCategoryText(AllObjectTypes, true), AllObjectTypes, this, LOCTEXT("ObjectType", "Object types"), true, TypesDatabasePtr)));
+			TSharedPtr<FPinTypeTreeInfo> Objects = MakeShared<FPinTypeTreeInfo>(GetCategoryText(AllObjectTypes, true), AllObjectTypes, this, LOCTEXT("ObjectType", "Object types"), true);
+			GatherPinsImpl::FindObjectsAndInterfaces(Objects, Interfaces);
+			TypeTree.Add(Objects);
 		}
 		else
 		{
-			TypeTree.Add(MakeShareable(new FPinTypeTreeInfo(GetCategoryText(PC_Object, true), PC_Object, this, LOCTEXT("ObjectTypeHardReference", "Hard reference to an Object"), true, TypesDatabasePtr)));
-			TypeTree.Add(MakeShareable(new FPinTypeTreeInfo(GetCategoryText(PC_Class, true), PC_Class, this, LOCTEXT("ClassType", "Hard reference to a Class"), true, TypesDatabasePtr)));
-			TypeTree.Add(MakeShareable(new FPinTypeTreeInfo(GetCategoryText(PC_SoftObject, true), PC_SoftObject, this, LOCTEXT("SoftObjectType", "Soft reference to an Object"), true, TypesDatabasePtr)));
-			TypeTree.Add(MakeShareable(new FPinTypeTreeInfo(GetCategoryText(PC_SoftClass, true), PC_SoftClass, this, LOCTEXT("SoftClassType", "Soft reference to a Class"), true, TypesDatabasePtr)));
+			TypeTree.Add(MakeShared<FPinTypeTreeInfo>(GetCategoryText(PC_Object, true), PC_Object, this, LOCTEXT("ObjectTypeHardReference", "Hard reference to an Object"), true));
+			TypeTree.Add(MakeShared<FPinTypeTreeInfo>(GetCategoryText(PC_Class, true), PC_Class, this, LOCTEXT("ClassType", "Hard reference to a Class"), true));
+			TypeTree.Add(MakeShared<FPinTypeTreeInfo>(GetCategoryText(PC_SoftObject, true), PC_SoftObject, this, LOCTEXT("SoftObjectType", "Soft reference to an Object"), true));
+			TypeTree.Add(MakeShared<FPinTypeTreeInfo>(GetCategoryText(PC_SoftClass, true), PC_SoftClass, this, LOCTEXT("SoftClassType", "Soft reference to a Class"), true));
 		}
 	}
-	TypeTree.Add( MakeShareable( new FPinTypeTreeInfo(GetCategoryText(PC_Enum, true), PC_Enum, this, LOCTEXT("EnumType", "Enumeration types."), true, TypesDatabasePtr) ) );
-
-#if SCHEMA_K2_GETVARIABLETYPETREE_LOG_TIME
-	const double EndTime = FPlatformTime::Seconds();
-	UE_LOG(LogBlueprint, Log, TEXT("UEdGraphSchema_K2::GetVariableTypeTree times - LoadedTypesDatabase: %f UnLoadedTypesDatabase: %f FPinTypeTreeInfo: %f"), DatabaseLoadedTime - StartTime, DatabaseUnLoadedTime - DatabaseLoadedTime, EndTime - DatabaseUnLoadedTime);
-#endif //SCHEMA_K2_GETVARIABLETYPETREE_LOG_TIME
-#undef SCHEMA_K2_GETVARIABLETYPETREE_LOG_TIME
+	TSharedPtr<FPinTypeTreeInfo> Enums = MakeShared<FPinTypeTreeInfo>(GetCategoryText(PC_Enum, true), PC_Enum, this, LOCTEXT("EnumType", "Enumeration types."), true);
+	if (!bRootTypesOnly)
+	{
+		GatherPinsImpl::FindEnums(Enums);
+	}
+	TypeTree.Add(Enums);
 }
 
 bool UEdGraphSchema_K2::DoesTypeHaveSubtypes(const FName Category) const
diff --git a/Engine/Source/Editor/KismetWidgets/Private/SPinTypeSelector.cpp b/Engine/Source/Editor/KismetWidgets/Private/SPinTypeSelector.cpp
--- a/Engine/Source/Editor/KismetWidgets/Private/SPinTypeSelector.cpp
+++ b/Engine/Source/Editor/KismetWidgets/Private/SPinTypeSelector.cpp
@@ -1,7 +1,10 @@
 // Copyright Epic Games, Inc. All Rights Reserved.
 #include "SPinTypeSelector.h"
 
+#include "Algo/AnyOf.h"
+#include "Algo/Count.h"
 #include "Algo/LevenshteinDistance.h"
+#include "Algo/RemoveIf.h"
 #include "BlueprintEditorSettings.h"
 #include "Containers/UnrealString.h"
 #include "EdGraph/EdGraphSchema.h"
@@ -349,6 +352,7 @@ void SPinTypeSelector::Construct(const FArguments& InArgs, FGetPinTypeTree GetPi
 	SelectorType = InArgs._SelectorType;
 
 	NumFilteredPinTypeItems = 0;
+	NumValidPinTypeItems = 0;
 
 	if(InArgs._CustomFilters.Num() > 0)
 	{
@@ -1062,40 +1066,8 @@ TSharedRef<SWidget>	SPinTypeSelector::GetMenuContent(bool bForSecondaryType)
 {
 	GetPinTypeTree.Execute(TypeTreeRoot, TypeTreeFilter);
 
-	// Remove read-only root items if they have no children; there will be no subtree to select non read-only items from in that case
-	int32 RootItemIndex = 0;
-	while(RootItemIndex < TypeTreeRoot.Num())
-	{
-		FPinTypeTreeItem TypeTreeItemPtr = TypeTreeRoot[RootItemIndex];
-		if(TypeTreeItemPtr.IsValid()
-			&& TypeTreeItemPtr->bReadOnly
-			&& TypeTreeItemPtr->Children.Num() == 0)
-		{
-			TypeTreeRoot.RemoveAt(RootItemIndex);
-		}
-		else
-		{
-			++RootItemIndex;
-		}
-	}
-
 	// Remove types not supported by schema
-	TArray<FPinTypeTreeItem> FilteredBySupportedTypes;
-	GetChildrenWithSupportedTypes(TypeTreeRoot, FilteredBySupportedTypes);
-	TypeTreeRoot = FilteredBySupportedTypes;
-
-	if (CustomFilters.Num() > 0)
-	{
-		NumFilteredPinTypeItems = 0;
-		FilteredTypeTreeRoot.Empty();
-
-		FTopLevenshteinResult<FPinTypeTreeItem> TopLevenshteinResult;
-		GetChildrenMatchingSearch(FText::GetEmpty(), TypeTreeRoot, FilteredTypeTreeRoot, TopLevenshteinResult);
-	}
-	else
-	{
-		FilteredTypeTreeRoot = TypeTreeRoot;
-	}
+	FilterUnsupportedTypes(TypeTreeRoot);
 
 	if( !MenuContent.IsValid() || (bForSecondaryType != bMenuContentIsSecondary) )
 	{
@@ -1399,8 +1371,8 @@ TOptional<T> FindBestMatch(FStringView SearchText, TArrayView<T> Container, TFun
 void SPinTypeSelector::OnFilterTextChanged(const FText& NewText)
 {
 	SearchText = NewText;
-	NumFilteredPinTypeItems = 0;
-	FilteredTypeTreeRoot.Empty();
+	NumFilteredPinTypeItems = NewText.IsEmpty() ? NumValidPinTypeItems : 0;
+	FilteredTypeTreeRoot.Reset();
 
 	FTopLevenshteinResult<FPinTypeTreeItem> TopLevenshteinResult;
 	GetChildrenMatchingSearch(NewText, TypeTreeRoot, FilteredTypeTreeRoot, TopLevenshteinResult);
@@ -1425,118 +1397,155 @@ void SPinTypeSelector::OnFilterTextCommitted(const FText& NewText, ETextCommit::
 	}
 }
 
-bool SPinTypeSelector::GetChildrenWithSupportedTypes(const TArray<FPinTypeTreeItem>& UnfilteredList,
-	TArray<FPinTypeTreeItem>& OutFilteredList)
+// helpers for dealing with the UEdGraphSchema_K2::FPinTypeTreeInfo tree:
+namespace PinTypeSelectorImpl
 {
-	bool bReturnVal = false;
-	
-	for( auto it = UnfilteredList.CreateConstIterator(); it; ++it )
-	{
-		FPinTypeTreeItem Item = *it;
-		FPinTypeTreeItem NewInfo = MakeShared<UEdGraphSchema_K2::FPinTypeTreeInfo>(Item);
-		TArray<FPinTypeTreeItem> ValidChildren;
+	template<typename T>
+	static void FilterUnsupportedTypesImpl(TArray<FPinTypeTreeItem>& ToFilter, const UEdGraphSchema* Schema, TWeakPtr<const FEdGraphSchemaAction> SchemaAction, T CustomFilter);
+	int32 CountValidTreeItemsImpl(const TArray<FPinTypeTreeItem>& ToCount);
+}
 
-		const bool bHasChildrenWithValidTypes = GetChildrenWithSupportedTypes(Item->Children, ValidChildren);
-		bool bSupportsType = true;
+template<typename T>
+static void PinTypeSelectorImpl::FilterUnsupportedTypesImpl(TArray<FPinTypeTreeItem>& ToFilter, const UEdGraphSchema* Schema, TWeakPtr<const FEdGraphSchemaAction> SchemaAction, T CustomFilter)
+{
+	ToFilter.SetNum(Algo::StableRemoveIf(ToFilter,
+		[Schema, SchemaAction, &CustomFilter = CustomFilter](const FPinTypeTreeItem& Item)
+		{
+			if (Item->Children.Num() > 0)
+			{
+				FilterUnsupportedTypesImpl<T>(Item->Children, Schema, SchemaAction, CustomFilter);
+			}
+			const bool bSupportsType = (Item->Children.Num() != 0 ||
+				(	Schema->SupportsPinType(SchemaAction, Item->GetPinTypeNoResolve()) &&
+					CustomFilter(Item)));
+			return !bSupportsType;
+		}
+	));
+}
 
-		if (!bHasChildrenWithValidTypes)
+int32 PinTypeSelectorImpl::CountValidTreeItemsImpl(const TArray<FPinTypeTreeItem>& ToCount)
+{
+	int32 Count = Algo::CountIf(ToCount, [](const FPinTypeTreeItem& Item) { return !Item->bReadOnly; });
+	for (const FPinTypeTreeItem& Item : ToCount)
+	{
+		Count += CountValidTreeItemsImpl(Item->Children);
+	}
+	return Count;
+}
+
+void SPinTypeSelector::FilterUnsupportedTypes(TArray<FPinTypeTreeItem>& ToFilter)
+{
+	// Remove read-only root items if they have no children; there will be no subtree to select non read-only items from in that case
+	int32 RootItemIndex = 0;
+	while (RootItemIndex < TypeTreeRoot.Num())
+	{
+		FPinTypeTreeItem TypeTreeItemPtr = TypeTreeRoot[RootItemIndex];
+		if (TypeTreeItemPtr.IsValid()
+			&& TypeTreeItemPtr->bReadOnly
+			&& TypeTreeItemPtr->Children.Num() == 0)
 		{
-			bSupportsType = Schema->SupportsPinType(SchemaAction, Item->GetPinTypeNoResolve());
+			TypeTreeRoot.RemoveAt(RootItemIndex);
 		}
-		
-		if (bHasChildrenWithValidTypes || bSupportsType)
+		else
 		{
-			NewInfo->Children = ValidChildren;
-			OutFilteredList.Add(NewInfo);
+			++RootItemIndex;
+		}
+	}
 
-			if (!NewInfo->bReadOnly)
-			{
-				++NumFilteredPinTypeItems;
-			}
+	const bool bHasCustomFilters = Algo::AnyOf(CustomFilters,
+		[](const TSharedPtr<IPinTypeSelectorFilter>& Filter)
+		{
+			return Filter.IsValid();
+		});
 
-			bReturnVal = true;
-		}
+	if (bHasCustomFilters)
+	{
+		PinTypeSelectorImpl::FilterUnsupportedTypesImpl(ToFilter, Schema, SchemaAction,
+			[&CustomFilters = CustomFilters](const FPinTypeTreeItem& Item)
+			{
+				bool bCustomFilterMatches = true;
+				for (const TSharedPtr<IPinTypeSelectorFilter>& CustomFilter : CustomFilters)
+				{
+					if (CustomFilter.IsValid())
+					{
+						bCustomFilterMatches &= CustomFilter->ShouldShowPinTypeTreeItem(Item);
+					}
+				}
+				return bCustomFilterMatches;
+			});
+	}
+	else
+	{
+		PinTypeSelectorImpl::FilterUnsupportedTypesImpl(ToFilter, Schema, SchemaAction, [](const FPinTypeTreeItem& Item) { return true; });
 	}
 
-	return bReturnVal;
+	// count items after filtering invalid ones, set up the 'filtered view' to default to all the valid items:
+	NumValidPinTypeItems = PinTypeSelectorImpl::CountValidTreeItemsImpl(ToFilter);
+	NumFilteredPinTypeItems = NumValidPinTypeItems;
+	FilteredTypeTreeRoot = TypeTreeRoot;
 }
 
 bool SPinTypeSelector::GetChildrenMatchingSearch(const FText& InSearchText, const TArray<FPinTypeTreeItem>& UnfilteredList, TArray<FPinTypeTreeItem>& OutFilteredList, FTopLevenshteinResult<FPinTypeTreeItem>& OutTopLevenshteinResult)
 {
+	const bool bIsEmptySearch = InSearchText.IsEmpty();
+	if (bIsEmptySearch)
+	{
+		OutFilteredList = UnfilteredList;
+		return true;
+	}
+
 	FString TrimmedFilterString;
 	TArray<FString> FilterTerms;
 	TArray<FString> SanitizedFilterTerms;
 
-	const bool bIsEmptySearch = InSearchText.IsEmpty();
-	if (!bIsEmptySearch)
-	{
-		// Trim and sanitized the filter text (so that it more likely matches the action descriptions)
-		TrimmedFilterString = FText::TrimPrecedingAndTrailing(InSearchText).ToString();
-
-		// Tokenize the search box text into a set of terms; all of them must be present to pass the filter
-		TrimmedFilterString.ParseIntoArray(FilterTerms, TEXT(" "), true);
+	// Trim and sanitized the filter text (so that it more likely matches the action descriptions)
+	TrimmedFilterString = FText::TrimPrecedingAndTrailing(InSearchText).ToString();
 
-		// Generate a list of sanitized versions of the strings
-		for (int32 iFilters = 0; iFilters < FilterTerms.Num(); iFilters++)
-		{
-			FString EachString = FName::NameToDisplayString(FilterTerms[iFilters], false);
-			EachString = EachString.Replace(TEXT(" "), TEXT(""));
-			SanitizedFilterTerms.Add(EachString);
-		}
+	// Tokenize the search box text into a set of terms; all of them must be present to pass the filter
+	TrimmedFilterString.ParseIntoArray(FilterTerms, TEXT(" "), true);
 
-		// Both of these should match!
-		ensure(SanitizedFilterTerms.Num() == FilterTerms.Num());
+	// Generate a list of sanitized versions of the strings
+	for (int32 iFilters = 0; iFilters < FilterTerms.Num(); iFilters++)
+	{
+		FString EachString = FName::NameToDisplayString(FilterTerms[iFilters], false);
+		EachString = EachString.Replace(TEXT(" "), TEXT(""));
+		SanitizedFilterTerms.Add(EachString);
 	}
 
+	// Both of these should match!
+	ensure(SanitizedFilterTerms.Num() == FilterTerms.Num());
+
 	bool bReturnVal = false;
 
 	for( const FPinTypeTreeItem& Item : UnfilteredList)
 	{
 		TArray<FPinTypeTreeItem> ValidChildren;
-		ValidChildren.Reserve(Item->Children.Num());
 		const bool bHasChildrenMatchingSearch = GetChildrenMatchingSearch(InSearchText, Item->Children, ValidChildren, OutTopLevenshteinResult);
-		bool bFilterMatches = bIsEmptySearch;
+		bool bFilterMatches = true;
 
 		// If children match the search filter, there's no need to do any additional checks
 		if (!bHasChildrenMatchingSearch)
 		{
-			// If valid, attempt to match the custom filters; otherwise, treat it as a match by default.
-			bool bHasValidCustomFilter = false;
-			bool bCustomFilterMatches = true;
-			for(const TSharedPtr<IPinTypeSelectorFilter>& CustomFilter : CustomFilters)
-			{
-				if(CustomFilter.IsValid())
-				{
-					bHasValidCustomFilter = true;
-					bCustomFilterMatches &= CustomFilter->ShouldShowPinTypeTreeItem(Item);
-				}
-			}
+			// but if they don't, we need to match the terms:
+			const FString& LocalizedDescriptionString = Item->GetCachedDescriptionString();
+			const FString* SourceDescriptionStringPtr = FTextInspector::GetSourceString(Item->GetDescription());
 
-			bFilterMatches = bHasValidCustomFilter ? bCustomFilterMatches : true;
-			
-			// If we didn't match the custom filter, or it's an empty search, let's not do any checks against the FilterTerms
-			if (bFilterMatches && !bIsEmptySearch)
-			{
-				const FString& LocalizedDescriptionString = Item->GetCachedDescriptionString();
-				const FString* SourceDescriptionStringPtr = FTextInspector::GetSourceString(Item->GetDescription());
-
-				// Test both the localized and source strings for a match
-				const FString MangledLocalizedDescriptionString = LocalizedDescriptionString.Replace(TEXT(" "), TEXT(""));
-				const FString MangledSourceDescriptionString = (SourceDescriptionStringPtr && *SourceDescriptionStringPtr != LocalizedDescriptionString) ? SourceDescriptionStringPtr->Replace(TEXT(" "), TEXT("")) : FString();
+			// Test both the localized and source strings for a match
+			const FString MangledLocalizedDescriptionString = LocalizedDescriptionString.Replace(TEXT(" "), TEXT(""));
+			const FString MangledSourceDescriptionString = (SourceDescriptionStringPtr && *SourceDescriptionStringPtr != LocalizedDescriptionString) ? SourceDescriptionStringPtr->Replace(TEXT(" "), TEXT("")) : FString();
 
-				for (int32 FilterIndex = 0; FilterIndex < FilterTerms.Num() && bFilterMatches; ++FilterIndex)
-				{
-					const bool bMatchesLocalizedTerm = MangledLocalizedDescriptionString.Contains(FilterTerms[FilterIndex]) || MangledLocalizedDescriptionString.Contains(SanitizedFilterTerms[FilterIndex]);
-					const bool bMatchesSourceTerm = !MangledSourceDescriptionString.IsEmpty() && (MangledSourceDescriptionString.Contains(FilterTerms[FilterIndex]) || MangledSourceDescriptionString.Contains(SanitizedFilterTerms[FilterIndex]));
-					bFilterMatches = bFilterMatches && (bMatchesLocalizedTerm || bMatchesSourceTerm);
-				}
+			for (int32 FilterIndex = 0; FilterIndex < FilterTerms.Num() && bFilterMatches; ++FilterIndex)
+			{
+				const bool bMatchesLocalizedTerm = MangledLocalizedDescriptionString.Contains(FilterTerms[FilterIndex]) || MangledLocalizedDescriptionString.Contains(SanitizedFilterTerms[FilterIndex]);
+				const bool bMatchesSourceTerm = !MangledSourceDescriptionString.IsEmpty() && (MangledSourceDescriptionString.Contains(FilterTerms[FilterIndex]) || MangledSourceDescriptionString.Contains(SanitizedFilterTerms[FilterIndex]));
+				bFilterMatches &= (bMatchesLocalizedTerm || bMatchesSourceTerm);
 			}
 		}
-		if( bHasChildrenMatchingSearch
-			|| bFilterMatches )
+
+		if(bFilterMatches)
 		{
 			FPinTypeTreeItem NewInfo = MakeShared<UEdGraphSchema_K2::FPinTypeTreeInfo>(Item);
-			NewInfo->Children = ValidChildren;
+			NewInfo->Children = MoveTemp(ValidChildren);
 			OutFilteredList.Add(NewInfo);
 
 			// Skip anything with children, those are categories and we don't care about those.
@@ -1705,13 +1714,35 @@ void SPinTypeSelector::OnMouseLeave( const FPointerEvent& MouseEvent )
 
 void SPinTypeSelector::OnCustomFilterChanged()
 {
-	NumFilteredPinTypeItems = 0;
-	FilteredTypeTreeRoot.Empty();
+	// regenerate clean type treeroot:
+	GetPinTypeTree.Execute(TypeTreeRoot, TypeTreeFilter);
 
-	FTopLevenshteinResult<FPinTypeTreeItem> TopLevenshteinResult;
-	GetChildrenMatchingSearch(SearchText, TypeTreeRoot, FilteredTypeTreeRoot, TopLevenshteinResult);
+	// apply new filters
+	FilterUnsupportedTypes(TypeTreeRoot);
 
-	if (TypeTreeView.IsValid())
+	// if any search text is present, apply that to the new FilteredTypeTreeRoot
+	if(!SearchText.IsEmpty())
+	{
+		NumFilteredPinTypeItems = 0;
+		FilteredTypeTreeRoot.Reset();
+
+		FTopLevenshteinResult<FPinTypeTreeItem> TopLevenshteinResult;
+		GetChildrenMatchingSearch(SearchText, TypeTreeRoot, FilteredTypeTreeRoot, TopLevenshteinResult);
+		
+		// FilteredTypeTreeRoot now up to date, we can refresh
+		if (TypeTreeView.IsValid())
+		{
+			TypeTreeView->RequestTreeRefresh();
+		}
+
+		// and update focus:
+		if (TopLevenshteinResult.IsSet())
+		{
+			TypeTreeView->SetSelection(TopLevenshteinResult.Item, ESelectInfo::OnNavigation);
+			TypeTreeView->RequestScrollIntoView(TopLevenshteinResult.Item);
+		}
+	}
+	else if (TypeTreeView.IsValid())
 	{
 		TypeTreeView->RequestTreeRefresh();
 	}
diff --git a/Engine/Source/Editor/KismetWidgets/Public/SPinTypeSelector.h b/Engine/Source/Editor/KismetWidgets/Public/SPinTypeSelector.h
--- a/Engine/Source/Editor/KismetWidgets/Public/SPinTypeSelector.h
+++ b/Engine/Source/Editor/KismetWidgets/Public/SPinTypeSelector.h
@@ -213,6 +213,9 @@ protected:
 	/** Total number of filtered pin type items. This count excludes category items and reference subtypes. */
 	int32 NumFilteredPinTypeItems;
 
+	/** Total number of valid pin type items, that could be filtered. This count excludes category items and reference subtypes. */
+	int32 NumValidPinTypeItems;
+
 	/** Holds a cache of the allowed Object Reference types for the last sub-menu opened. */
 	TArray<FObjectReferenceListItem> AllowedObjectReferenceTypes;
 	TWeakPtr<SListView<FObjectReferenceListItem>> WeakListView;
@@ -259,7 +262,7 @@ protected:
 	void OnFilterTextCommitted(const FText& NewText, ETextCommit::Type CommitInfo);
 
 	/** Helper to generate the filtered list of types, based on the supported types of the schema */
-	bool GetChildrenWithSupportedTypes(const TArray<FPinTypeTreeItem>& UnfilteredList, TArray<FPinTypeTreeItem>& OutFilteredList);
+	void FilterUnsupportedTypes(TArray<FPinTypeTreeItem>& UnfilteredList);
 
 	/** Helper to generate the filtered list of types, based on the search string matching */
 	bool GetChildrenMatchingSearch(const FText& SearchText, const TArray<FPinTypeTreeItem>& UnfilteredList, TArray<FPinTypeTreeItem>& OutFilteredList, FTopLevenshteinResult<FPinTypeTreeItem>& OutTopLevenshteinResult);
