Added support for autocast functions for objects to support their super classes

Example:

Class A {};
Class B : public A {};
Class C {};
Class D : public C{};

if auto cast from D to B exists, it means auto cast from D to A also exists
if auto cast from C to B exists, it means auto cast from D to B also exists

#changelist validated
#virtualized

diff --git a/Engine/Source/Editor/BlueprintGraph/Private/EdGraphSchema_K2.cpp b/Engine/Source/Editor/BlueprintGraph/Private/EdGraphSchema_K2.cpp
--- a/Engine/Source/Editor/BlueprintGraph/Private/EdGraphSchema_K2.cpp
+++ b/Engine/Source/Editor/BlueprintGraph/Private/EdGraphSchema_K2.cpp
@@ -2367,7 +2367,22 @@ private:
 		FEdGraphPinType OutputPinType;
 		Schema->ConvertPropertyToPinType(Function->GetReturnProperty(), OutputPinType);
 
-		InnerMap.Add(GenerateCastData(InputPinType, OutputPinType), Function);
+		// If the output pin is an object pin, iterate through all possible super classes to add them as viable auto cast functions
+		UStruct* StructObject = Cast<UStruct>(OutputPinType.PinSubCategoryObject.Get());
+		const bool bIterateHierarchy = OutputPinType.PinCategory == UEdGraphSchema_K2::PC_Object && StructObject;
+		if (bIterateHierarchy)
+		{
+			FEdGraphPinType OutputPinTypeCopy = OutputPinType;
+			for (UStruct* OutputPinObject = StructObject; OutputPinObject != nullptr; OutputPinObject = OutputPinObject->GetSuperStruct())
+			{
+				OutputPinTypeCopy.PinSubCategoryObject = OutputPinObject;
+				InnerMap.Add(GenerateCastData(InputPinType, OutputPinTypeCopy), Function);
+			}
+		}
+		else
+		{
+			InnerMap.Add(GenerateCastData(InputPinType, OutputPinType), Function);
+		}
 	}
 public:
 
@@ -2452,6 +2467,25 @@ public:
 
 	UFunction* Find(const FEdGraphPinType& InputPinType, const FEdGraphPinType& OutputPinType) const
 	{
+		// If the input pin is an object pin, iterate through all possible super classes to check for auto cast availability
+		UStruct* StructObject = Cast<UStruct>(InputPinType.PinSubCategoryObject.Get());
+		const bool bIterateHierarchy = InputPinType.PinCategory == UEdGraphSchema_K2::PC_Object && StructObject;
+		if (bIterateHierarchy)
+		{
+			FEdGraphPinType InputPinTypeCopy = InputPinType;
+			for (UStruct* InputPinObject = StructObject; InputPinObject != nullptr; InputPinObject = InputPinObject->GetSuperStruct())
+			{
+				InputPinTypeCopy.PinSubCategoryObject = InputPinObject;
+				
+				const TWeakObjectPtr<UFunction>* FuncPtr = InnerMap.Find(GenerateCastData(InputPinTypeCopy, OutputPinType));
+				if (FuncPtr)
+				{
+					return FuncPtr->Get();
+				}
+			}
+			return nullptr;
+		}
+		
 		const TWeakObjectPtr<UFunction>* FuncPtr = InnerMap.Find(GenerateCastData(InputPinType, OutputPinType));
 		return FuncPtr ? FuncPtr->Get() : nullptr;
 	}
