Fix enum pins in the pin type selector, broken by 31223174, which failed to convert PC_Enum to PC_Byte

diff --git a/Engine/Source/Editor/BlueprintGraph/Private/EdGraphSchema_K2.cpp b/Engine/Source/Editor/BlueprintGraph/Private/EdGraphSchema_K2.cpp
--- a/Engine/Source/Editor/BlueprintGraph/Private/EdGraphSchema_K2.cpp
+++ b/Engine/Source/Editor/BlueprintGraph/Private/EdGraphSchema_K2.cpp
@@ -343,7 +343,7 @@ void GatherPinsImpl::FindEnums(const TSharedPtr<UEdGraphSchema_K2::FPinTypeTreeI
 		{
 			Owner->Children.Emplace(
 				FromObject(CurrentEnum
-					, UEdGraphSchema_K2::PC_Enum
+					, UEdGraphSchema_K2::PC_Byte
 					, EObjectReferenceType::NotAnObject));
 		}
 	}
@@ -358,7 +358,7 @@ void GatherPinsImpl::FindEnums(const TSharedPtr<UEdGraphSchema_K2::FPinTypeTreeI
 		{
 			Owner->Children.Emplace(
 				FromAssetData(Asset
-					, UEdGraphSchema_K2::PC_Enum
+					, UEdGraphSchema_K2::PC_Byte
 					, EObjectReferenceType::NotAnObject));
 		}
 	}
