Function parameter and local variable nodes in Blueprint function graphs now have an icon above them. This can be changed through Editor Settings.

diff --git a/Engine/Source/Editor/BlueprintGraph/Private/BlueprintEditorSettings.cpp b/Engine/Source/Editor/BlueprintGraph/Private/BlueprintEditorSettings.cpp
--- a/Engine/Source/Editor/BlueprintGraph/Private/BlueprintEditorSettings.cpp
+++ b/Engine/Source/Editor/BlueprintGraph/Private/BlueprintEditorSettings.cpp
@@ -32,6 +32,8 @@ UBlueprintEditorSettings::UBlueprintEditorSettings(const FObjectInitializer& Obj
 	, bShowGraphInstructionText(true)
 	, bHideUnrelatedNodes(false)
 	, bShowShortTooltips(true)
+	, bShowFunctionParameterIcon(true)
+	, bShowFunctionLocalVariableIcon(true)
 	// Workflow Settings
 	, bEnableInputTriggerSupportWarnings(false)
 	, bSplitContextTargetSettings(true)
diff --git a/Engine/Source/Editor/BlueprintGraph/Private/K2Node_Variable.cpp b/Engine/Source/Editor/BlueprintGraph/Private/K2Node_Variable.cpp
--- a/Engine/Source/Editor/BlueprintGraph/Private/K2Node_Variable.cpp
+++ b/Engine/Source/Editor/BlueprintGraph/Private/K2Node_Variable.cpp
@@ -4,6 +4,7 @@
 #include "K2Node_Variable.h"
 
 #include "BlueprintCompilationManager.h"
+#include "BlueprintEditorSettings.h"
 #include "Components/ActorComponent.h"
 #include "Components/PrimitiveComponent.h"
 #include "Containers/Set.h"
@@ -640,6 +641,9 @@ FText UK2Node_Variable::GetToolTipHeading() const
 		const UActorComponent* Component = GetActorComponent(VariableProperty);
 		const bool IsEditorOnly = VariableProperty->HasAnyPropertyFlags(CPF_EditorOnly) || (Component && Component->bIsEditorOnly);
 		const bool IsReplicated = VariableProperty->HasAnyPropertyFlags(CPF_Net) || (Component && Component->GetIsReplicated());
+		const bool IsFunctionParameter = VariableProperty->HasAnyPropertyFlags(CPF_Parm);
+		const bool IsLocalFunctionVariable = !VariableProperty->HasAnyPropertyFlags(CPF_Parm) && VariableReference.IsLocalScope();
+		const UBlueprintEditorSettings* EditorSettings = GetDefault<UBlueprintEditorSettings>();
 		if (IsEditorOnly && IsReplicated)
 		{
 			IconTag = LOCTEXT("ReplicatedEditorOnlyVar", "Editor-Only | Replicated");
@@ -652,6 +656,14 @@ FText UK2Node_Variable::GetToolTipHeading() const
 		{
 			IconTag = LOCTEXT("EditorOnlyVar", "Editor-Only");
 		}
+		else if (IsFunctionParameter && EditorSettings->bShowFunctionParameterIcon)
+		{
+			IconTag = LOCTEXT("FunctionParameterVar", "Function parameter");
+		}
+		else if (IsLocalFunctionVariable && EditorSettings->bShowFunctionLocalVariableIcon)
+		{
+			IconTag = LOCTEXT("LocalFunctionVar", "Function local variable");
+		}
 	}
 
 	if (Heading.IsEmpty())
@@ -878,14 +890,27 @@ FName UK2Node_Variable::GetCornerIcon() const
 	if (const FProperty* VariableProperty = VariableReference.ResolveMember<FProperty>(GetBlueprintClassFromNode()))
 	{
 		const UActorComponent* Component = GetActorComponent(VariableProperty);
-		if (VariableProperty->HasAllPropertyFlags(CPF_Net) || (Component && Component->GetIsReplicated()))
+		const bool IsEditorOnly = VariableProperty->HasAnyPropertyFlags(CPF_EditorOnly) || (Component && Component->bIsEditorOnly);
+		const bool IsReplicated = VariableProperty->HasAnyPropertyFlags(CPF_Net) || (Component && Component->GetIsReplicated());
+		const bool IsFunctionParameter = VariableProperty->HasAnyPropertyFlags(CPF_Parm);
+		const bool IsLocalFunctionVariable = !VariableProperty->HasAnyPropertyFlags(CPF_Parm) && VariableReference.IsLocalScope();
+		const UBlueprintEditorSettings* EditorSettings = GetDefault<UBlueprintEditorSettings>();
+		if (IsReplicated)
 		{
 			return TEXT("Graph.Replication.Replicated");
 		}
-		else if (VariableProperty->HasAllPropertyFlags(CPF_EditorOnly) || (Component && Component->bIsEditorOnly))
+		else if (IsEditorOnly)
 		{
 			return TEXT("Graph.Editor.EditorOnlyIcon");
 		}
+		else if (IsFunctionParameter && EditorSettings->bShowFunctionParameterIcon)
+		{
+			return TEXT("Graph.Function.FunctionParameterIcon");
+		}
+		else if (IsLocalFunctionVariable && EditorSettings->bShowFunctionLocalVariableIcon)
+		{
+			return TEXT("Graph.Function.FunctionLocalVariableIcon");
+		}
 	}
 
 	return Super::GetCornerIcon();
diff --git a/Engine/Source/Editor/BlueprintGraph/Public/BlueprintEditorSettings.h b/Engine/Source/Editor/BlueprintGraph/Public/BlueprintEditorSettings.h
--- a/Engine/Source/Editor/BlueprintGraph/Public/BlueprintEditorSettings.h
+++ b/Engine/Source/Editor/BlueprintGraph/Public/BlueprintEditorSettings.h
@@ -107,6 +107,14 @@ public:
 	UPROPERTY(EditAnywhere, config, Category = VisualStyle)
 	bool bShowShortTooltips;
 
+	/** If true, show an icon next to function parameter nodes. */
+	UPROPERTY(EditAnywhere, config, Category = VisualStyle)
+	bool bShowFunctionParameterIcon;
+
+	/** If true, show an icon next to function local variable nodes. */
+	UPROPERTY(EditAnywhere, config, Category = VisualStyle)
+	bool bShowFunctionLocalVariableIcon;
+
 // Workflow Settings
 public:
 	/** 
diff --git a/Engine/Source/Editor/EditorStyle/Private/StarshipStyle.cpp b/Engine/Source/Editor/EditorStyle/Private/StarshipStyle.cpp
--- a/Engine/Source/Editor/EditorStyle/Private/StarshipStyle.cpp
+++ b/Engine/Source/Editor/EditorStyle/Private/StarshipStyle.cpp
@@ -3793,6 +3793,8 @@ void FStarshipEditorStyle::FStyle::SetupGraphEditorStyles()
 
 		Set( "Graph.Latent.LatentIcon", new IMAGE_BRUSH("/Graph/LatentIcon", FVector2D(32,32) ) );
 		Set( "Graph.Message.MessageIcon", new IMAGE_BRUSH("/Graph/MessageIcon", FVector2D(32,32) ) );
+		Set( "Graph.Function.FunctionParameterIcon", new IMAGE_BRUSH_SVG("/Starship/GraphEditors/FunctionInputParameter", FVector2D(20.0f, 20.0f) ) );
+		Set( "Graph.Function.FunctionLocalVariableIcon", new IMAGE_BRUSH_SVG("/Starship/GraphEditors/FunctionLocalVariable", FVector2D(20.0f, 20.0f) ) );
 
 		Set( "Graph.ExecPin.Connected", new IMAGE_BRUSH( "Old/Graph/ExecPin_Connected", Icon12x16 ) );
 		Set( "Graph.ExecPin.Disconnected", new IMAGE_BRUSH( "Old/Graph/ExecPin_Disconnected", Icon12x16 ) );
