Downgrade check to compile error

diff --git a/Engine/Source/Editor/BlueprintGraph/Private/K2Node_VariableGet.cpp b/Engine/Source/Editor/BlueprintGraph/Private/K2Node_VariableGet.cpp
--- a/Engine/Source/Editor/BlueprintGraph/Private/K2Node_VariableGet.cpp
+++ b/Engine/Source/Editor/BlueprintGraph/Private/K2Node_VariableGet.cpp
@@ -549,7 +549,11 @@ void UK2Node_VariableGet::ExpandNode(class FKismetCompilerContext& CompilerConte
 		{
 			UClass* OwnerClass = VariableProperty->GetOwnerClass();
 			UFunction* GetFunction = OwnerClass->FindFunctionByName(*GetFunctionName);
-			check(GetFunction);
+			if (!GetFunction)
+			{
+				CompilerContext.MessageLog.Error(*LOCTEXT("MissingGetter", "Getter function not found for @@").ToString(), this);
+				return;
+			}
 
 			UK2Node_CallFunction* CallFuncNode = CompilerContext.SpawnIntermediateNode<UK2Node_CallFunction>(this, SourceGraph);
 			CallFuncNode->SetFromFunction(GetFunction);
diff --git a/Engine/Source/Editor/BlueprintGraph/Private/K2Node_VariableSet.cpp b/Engine/Source/Editor/BlueprintGraph/Private/K2Node_VariableSet.cpp
--- a/Engine/Source/Editor/BlueprintGraph/Private/K2Node_VariableSet.cpp
+++ b/Engine/Source/Editor/BlueprintGraph/Private/K2Node_VariableSet.cpp
@@ -503,7 +503,11 @@ void UK2Node_VariableSet::ExpandNode(class FKismetCompilerContext& CompilerConte
 			{
 				UClass* OwnerClass = VariableProperty->GetOwnerClass();
 				UFunction* SetFunction = OwnerClass->FindFunctionByName(*SetFunctionName);
-				check(SetFunction);
+				if (!SetFunction)
+				{
+					CompilerContext.MessageLog.Error(*LOCTEXT("MissingSetter", "Setter function not found for @@").ToString(), this);
+					return;
+				}
 
 				UK2Node_CallFunction* CallFuncNode = CompilerContext.SpawnIntermediateNode<UK2Node_CallFunction>(this, SourceGraph);
 				CallFuncNode->SetFromFunction(SetFunction);
