Switch node raises an error when its enum disapears

diff --git a/Engine/Source/Editor/BlueprintGraph/Classes/K2Node_SwitchEnum.h b/Engine/Source/Editor/BlueprintGraph/Classes/K2Node_SwitchEnum.h
--- a/Engine/Source/Editor/BlueprintGraph/Classes/K2Node_SwitchEnum.h
+++ b/Engine/Source/Editor/BlueprintGraph/Classes/K2Node_SwitchEnum.h
@@ -52,6 +52,7 @@ class UK2Node_SwitchEnum : public UK2Node_Switch, public INodeDependingOnEnumInt
 	virtual FText GetTooltipText() const override;
 	virtual FText GetNodeTitle(ENodeTitleType::Type TitleType) const override;
 	virtual void AddPinSearchMetaDataInfo(const UEdGraphPin* Pin, TArray<struct FSearchTagDataPair>& OutTaggedMetaData) const override;
+	virtual void ValidateNodeDuringCompilation(FCompilerResultsLog& MessageLog) const override;
 	// End of UEdGraphNode interface
 
 	// UK2Node interface
diff --git a/Engine/Source/Editor/BlueprintGraph/Private/K2Node_SwitchEnum.cpp b/Engine/Source/Editor/BlueprintGraph/Private/K2Node_SwitchEnum.cpp
--- a/Engine/Source/Editor/BlueprintGraph/Private/K2Node_SwitchEnum.cpp
+++ b/Engine/Source/Editor/BlueprintGraph/Private/K2Node_SwitchEnum.cpp
@@ -109,6 +109,15 @@ void UK2Node_SwitchEnum::AddPinSearchMetaDataInfo(const UEdGraphPin* Pin, TArray
 	}
 }
 
+void UK2Node_SwitchEnum::ValidateNodeDuringCompilation(FCompilerResultsLog& MessageLog) const
+{
+	if (!Enum)
+	{
+		MessageLog.Error(*LOCTEXT("MissingSwitchEnum", "@@ must have a valid enum - check log for 'Failed import' add a redirector or recreate the enum").ToString(), this);
+	}
+	return Super::ValidateNodeDuringCompilation(MessageLog);
+}
+
 bool UK2Node_SwitchEnum::IsConnectionDisallowed(const UEdGraphPin* MyPin, const UEdGraphPin* OtherPin, FString& OutReason) const
 {
 	const UEnum* SubCategoryObject = Cast<UEnum>( OtherPin->PinType.PinSubCategoryObject.Get() );
