Deprecated redundant fields in UK2Node_CallFunction.

Both bIsConstFunc and bIsInterfaceCall are unnecessary since we can simply query the underlying UFunction/UClass (via FunctionReference) to determine if we're dealing with a const function or an interface class.

Additional changes:
* A check for CLASS_Interface was added to UK2Node_Message. While the chances of hitting this are very low, it wouldn't make sense to reference a non-interface class when expanding this node. It should really error out in that case.
* CreatePinsForFunctionCall had an unusual check for hiding the self pin if the pure function was non-const. This appears to be really old and no longer relevant.

diff --git a/Engine/Source/Editor/BlueprintGraph/Classes/K2Node_CallFunction.h b/Engine/Source/Editor/BlueprintGraph/Classes/K2Node_CallFunction.h
--- a/Engine/Source/Editor/BlueprintGraph/Classes/K2Node_CallFunction.h
+++ b/Engine/Source/Editor/BlueprintGraph/Classes/K2Node_CallFunction.h
@@ -43,16 +43,14 @@ class BLUEPRINTGRAPH_API UK2Node_CallFunction : public UK2Node
 	UPROPERTY()
 	uint32 bIsPureFunc:1;
 
-	/** Indicates that this is a call to a const function */
-	UPROPERTY()
-	uint32 bIsConstFunc:1;
-
 	/** Indicates that during compile we want to create multiple exec pins from an enum param */
 	UPROPERTY()
 	uint32 bWantsEnumToExecExpansion:1;
 
-	/** Indicates that this is a call to an interface function */
-	UPROPERTY()
+	UE_DEPRECATED(5.5, "bIsConstFunc is deprecated. Check for FUNC_Const on FunctionReference.ResolveMember<UFunction>() instead.")
+	uint32 bIsConstFunc:1;
+
+	UE_DEPRECATED(5.5, "bIsInterfaceCall is deprecated. Check for CLASS_Interface on FunctionReference.GetMemberParentClass() instead.")
 	uint32 bIsInterfaceCall:1;
 
 	UE_DEPRECATED(5.4, "bIsFinalFunction is deprecated.")
diff --git a/Engine/Source/Editor/BlueprintGraph/Private/CallFunctionHandler.cpp b/Engine/Source/Editor/BlueprintGraph/Private/CallFunctionHandler.cpp
--- a/Engine/Source/Editor/BlueprintGraph/Private/CallFunctionHandler.cpp
+++ b/Engine/Source/Editor/BlueprintGraph/Private/CallFunctionHandler.cpp
@@ -474,12 +474,16 @@ void FKCHandler_CallFunction::CreateFunctionCallStatement(FKismetFunctionContext
 				// of a cast node). Otherwise, we'll infer the wrong context type at runtime and corrupt the stack by
 				// reading an interface ptr (16 bytes) into an object ptr (8 bytes) when we process the context opcode.
 				const bool bIsInterfaceContextTerm = Target && Target->AssociatedVarProperty && Target->AssociatedVarProperty->IsA<FInterfaceProperty>();
+				
+				UClass* FunctionOwnerClass = Function->GetOuterUClass();
+				check(FunctionOwnerClass);
+				const bool bIsInterfaceFunc = FunctionOwnerClass->HasAnyClassFlags(CLASS_Interface);
 
 				FBlueprintCompiledStatement& Statement = Context.AppendStatementForNode(Node);
 				Statement.FunctionToCall = Function;
 				Statement.FunctionContext = Target;
 				Statement.Type = KCST_CallFunction;
-				Statement.bIsInterfaceContext = IsCalledFunctionFromInterface(Node) || bIsInterfaceContextTerm;
+				Statement.bIsInterfaceContext = bIsInterfaceFunc || bIsInterfaceContextTerm;
 				Statement.bIsParentContext = Node->IsA<UK2Node_CallParentFunction>();
 
 				Statement.LHS = LHSTerm;
diff --git a/Engine/Source/Editor/BlueprintGraph/Private/K2Node_CallFunction.cpp b/Engine/Source/Editor/BlueprintGraph/Private/K2Node_CallFunction.cpp
--- a/Engine/Source/Editor/BlueprintGraph/Private/K2Node_CallFunction.cpp
+++ b/Engine/Source/Editor/BlueprintGraph/Private/K2Node_CallFunction.cpp
@@ -1124,9 +1124,7 @@ bool UK2Node_CallFunction::CreatePinsForFunctionCall(const UFunction* Function)
 
 	UClass* FunctionOwnerClass = Function->GetOuterUClass();
 
-	bIsInterfaceCall = FunctionOwnerClass->HasAnyClassFlags(CLASS_Interface);
 	bIsPureFunc = (Function->HasAnyFunctionFlags(FUNC_BlueprintPure) != false);
-	bIsConstFunc = (Function->HasAnyFunctionFlags(FUNC_Const) != false);
 	DetermineWantsEnumToExecExpansion(Function);
 
 	// Create input pins
@@ -1168,8 +1166,9 @@ bool UK2Node_CallFunction::CreatePinsForFunctionCall(const UFunction* Function)
 			}
 			else
 			{
-				// Hide the self pin if the function is compatible with the blueprint class and pure (the !bIsConstFunc portion should be going away soon too hopefully)
-				SelfPin->bHidden = (bIsFunctionCompatibleWithSelf && bIsPureFunc && !bIsConstFunc);
+				// Pure functions are generally compact looking, so we take advantage of that by hiding the self pin if we can.
+				// In this case, if the function belongs to our current class, then "self" is implied.
+				SelfPin->bHidden = (bIsFunctionCompatibleWithSelf && bIsPureFunc);
 			}
 		}
 	}
@@ -1933,7 +1932,6 @@ void UK2Node_CallFunction::SetFromFunction(const UFunction* Function)
 	if (Function != NULL)
 	{
 		bIsPureFunc = Function->HasAnyFunctionFlags(FUNC_BlueprintPure);
-		bIsConstFunc = Function->HasAnyFunctionFlags(FUNC_Const);
 		DetermineWantsEnumToExecExpansion(Function);
 
 		FunctionReference.SetFromField<UFunction>(Function, GetBlueprintClassFromNode());
diff --git a/Engine/Source/Editor/BlueprintGraph/Private/K2Node_CallParentFunction.cpp b/Engine/Source/Editor/BlueprintGraph/Private/K2Node_CallParentFunction.cpp
--- a/Engine/Source/Editor/BlueprintGraph/Private/K2Node_CallParentFunction.cpp
+++ b/Engine/Source/Editor/BlueprintGraph/Private/K2Node_CallParentFunction.cpp
@@ -68,7 +68,6 @@ void UK2Node_CallParentFunction::SetFromFunction(const UFunction* Function)
 	if (Function != NULL)
 	{
 		bIsPureFunc = Function->HasAnyFunctionFlags(FUNC_BlueprintPure);
-		bIsConstFunc = Function->HasAnyFunctionFlags(FUNC_Const);
 
 		UClass* OwnerClass = Function->GetOwnerClass();
 
diff --git a/Engine/Source/Editor/BlueprintGraph/Private/K2Node_Message.cpp b/Engine/Source/Editor/BlueprintGraph/Private/K2Node_Message.cpp
--- a/Engine/Source/Editor/BlueprintGraph/Private/K2Node_Message.cpp
+++ b/Engine/Source/Editor/BlueprintGraph/Private/K2Node_Message.cpp
@@ -188,21 +188,29 @@ void UK2Node_Message::ExpandNode(class FKismetCompilerContext& CompilerContext,
 	// Skip ourselves if our exec isn't wired up
 	if (bExecPinConnected)
 	{
+		UClass* InterfaceClass = FunctionReference.GetMemberParentClass(GetBlueprintClassFromNode());
+
 		// Make sure our interface is valid
-		if (FunctionReference.GetMemberParentClass(GetBlueprintClassFromNode()) == NULL)
+		if (InterfaceClass == nullptr)
 		{
 			CompilerContext.MessageLog.Error(*LOCTEXT("MessageNodeInvalid_Error", "Message node @@ has an invalid interface.").ToString(), this);
 			return;
 		}
 
+		if (!InterfaceClass->HasAnyClassFlags(CLASS_Interface))
+		{
+			CompilerContext.MessageLog.Error(*LOCTEXT("MessageNodeClassNotAnInterface_Error", "Message node @@ uses a class that isn't an interface.").ToString(), this);
+			return;
+		}
+
 		UFunction* MessageNodeFunction = GetTargetFunction();
-		if (MessageNodeFunction == NULL)
+		if (MessageNodeFunction == nullptr)
 		{
 			//@TODO: Why do this here in the compiler, it's already done on AllocateDefaultPins() during on-load node reconstruction
-			MessageNodeFunction = FMemberReference::FindRemappedField<UFunction>(FunctionReference.GetMemberParentClass(GetBlueprintClassFromNode()), FunctionReference.GetMemberName());
+			MessageNodeFunction = FMemberReference::FindRemappedField<UFunction>(InterfaceClass, FunctionReference.GetMemberName());
 		}
 
-		if (MessageNodeFunction == NULL)
+		if (MessageNodeFunction == nullptr)
 		{
 			CompilerContext.MessageLog.Error(*FText::Format(LOCTEXT("MessageNodeInvalidFunction_ErrorFmt", "Unable to find function with name {0} for Message node @@."), FText::FromString(FunctionReference.GetMemberName().ToString())).ToString(), this);
 			return;
@@ -255,7 +263,6 @@ void UK2Node_Message::ExpandNode(class FKismetCompilerContext& CompilerContext,
 
 		// Next, create the function call node
 		UK2Node_CallFunction* FunctionCallNode = CompilerContext.SpawnIntermediateNode<UK2Node_CallFunction>(this, SourceGraph);
-		FunctionCallNode->bIsInterfaceCall = true;
 		FunctionCallNode->FunctionReference = FunctionReference;
 		FunctionCallNode->AllocateDefaultPins();
 
@@ -285,7 +292,8 @@ void UK2Node_Message::ExpandNode(class FKismetCompilerContext& CompilerContext,
 		UFunction* MapClearFunction = UBlueprintMapLibrary::StaticClass()->FindFunctionByName(FName(TEXT("Map_Clear")));
 		check(MapClearFunction);
 
-		bool const bIsPureMessageFunc = Super::IsNodePure();
+		const bool bIsPureMessageFunc = Super::IsNodePure();
+
 		// Variable pins - Try to associate variable inputs to the message node with the variable inputs and outputs to the call function node
 		for( int32 i = 0; i < Pins.Num(); i++ )
 		{
diff --git a/Engine/Source/Editor/BlueprintGraph/Public/CallFunctionHandler.h b/Engine/Source/Editor/BlueprintGraph/Public/CallFunctionHandler.h
--- a/Engine/Source/Editor/BlueprintGraph/Public/CallFunctionHandler.h
+++ b/Engine/Source/Editor/BlueprintGraph/Public/CallFunctionHandler.h
@@ -51,13 +51,9 @@ public:
 		return false;
 	}
 
+	UE_DEPRECATED(5.5, "IsCalledFunctionFromInterface is deprecated")
 	bool IsCalledFunctionFromInterface(UEdGraphNode* Node)
 	{
-		if (UK2Node_CallFunction* CallFunctionNode = Cast<UK2Node_CallFunction>(Node))
-		{
-			return CallFunctionNode->bIsInterfaceCall;
-		}
-
 		return false;
 	}
 
