Avoid use of TWeakFieldPtr for performance reasons, avoid use TWeakObjectPtr because hashing it in this way is unsafe. This is principally an optimization for the blueprint context menu, as this binding object is frequently constructed and is passed by value

diff --git a/Engine/Source/Editor/BlueprintGraph/Classes/BlueprintNodeBinder.h b/Engine/Source/Editor/BlueprintGraph/Classes/BlueprintNodeBinder.h
--- a/Engine/Source/Editor/BlueprintGraph/Classes/BlueprintNodeBinder.h
+++ b/Engine/Source/Editor/BlueprintGraph/Classes/BlueprintNodeBinder.h
@@ -12,8 +12,8 @@ class UEdGraphNode;
 
 class FBindingObject
 {
-	TWeakObjectPtr<UObject> Object;
-	TWeakFieldPtr<FField> Property;
+	TObjectPtr<UObject> Object;
+	FField* Field;
 	bool bIsUObject;
 public:
 	FBindingObject()
@@ -25,10 +25,10 @@ public:
 		, bIsUObject(true)
 	{}
 	FBindingObject(FField* InField)
-		: Property(InField)
+		: Field(InField)
 		, bIsUObject(false)
 	{}
-	FBindingObject(FFieldVariant InFieldOrObject)
+	FBindingObject(const FFieldVariant& InFieldOrObject)
 	{
 		bIsUObject = InFieldOrObject.IsUObject();
 		if (bIsUObject)
@@ -37,28 +37,28 @@ public:
 		}
 		else
 		{
-			Property = InFieldOrObject.ToField();
+			Field = InFieldOrObject.ToField();
 		}
 	}
 	template <typename T, decltype(ImplicitConv<UObject*>(DeclVal<T>()))* = nullptr>
 	FBindingObject& operator=(T InObject)
 	{
 		Object = ImplicitConv<UObject*>(InObject);
-		Property = nullptr;
+		Field = nullptr;
 		bIsUObject = true;
 		return *this;
 	}
 	FBindingObject& operator=(FField* InField)
 	{
 		Object = nullptr;
-		Property = InField;
+		Field = InField;
 		bIsUObject = false;
 		return *this;
 	}
 	FBindingObject& operator=(TYPE_OF_NULLPTR)
 	{
 		Object = nullptr;
-		Property = nullptr;
+		Field = nullptr;
 		return *this;
 	}
 	bool IsUObject() const
@@ -67,47 +67,47 @@ public:
 	}
 	bool IsValid() const
 	{
-		return bIsUObject ? Object.IsValid() : Property.IsValid();
+		return bIsUObject ? Object != nullptr : Field != nullptr;
 	}
 	FName GetFName() const
 	{
-		return bIsUObject ? Object->GetFName() : Property->GetFName();
+		return bIsUObject ? Object->GetFName() : Field->GetFName();
 	}
 	FString GetName() const
 	{
-		return bIsUObject ? Object->GetName() : Property->GetName();
+		return bIsUObject ? Object->GetName() : Field->GetName();
 	}
 	FString GetPathName() const
 	{
-		return bIsUObject ? Object->GetPathName() : Property->GetPathName();
+		return bIsUObject ? Object->GetPathName() : Field->GetPathName();
 	}
 	FString GetFullName() const
 	{
-		return bIsUObject ? Object->GetFullName() : Property->GetFullName();
+		return bIsUObject ? Object->GetFullName() : Field->GetFullName();
 	}
 	bool IsA(const UClass* InClass) const
 	{
-		return bIsUObject && Object.IsValid() && Object->IsA(InClass);
+		return bIsUObject && Object != nullptr && Object->IsA(InClass);
 	}
 	bool IsA(const FFieldClass* InClass) const
 	{
-		return !bIsUObject && Property.IsValid() && Property->IsA(InClass);
+		return !bIsUObject && Field != nullptr && Field->IsA(InClass);
 	}
 	template <typename T>
 	bool IsA() const
 	{
 		if constexpr (std::is_base_of_v<UObject, T>)
 		{
-			if (bIsUObject && Object.IsValid())
+			if (bIsUObject && Object)
 			{
 				return Object->IsA(T::StaticClass());
 			}
 		}
 		else
 		{
-			if (!bIsUObject && Property.IsValid())
+			if (!bIsUObject && Field != nullptr)
 			{
-				return Property->IsA(T::StaticClass());
+				return Field->IsA(T::StaticClass());
 			}
 		}
 		return false;
@@ -117,16 +117,16 @@ public:
 	{
 		if constexpr (std::is_base_of_v<UObject, T>)
 		{
-			if (bIsUObject && Object.IsValid())
+			if (bIsUObject && Object)
 			{
-				return Cast<T>(Object.Get());
+				return Cast<T>(Object);
 			}
 		}
 		else
 		{
-			if (!bIsUObject && Property.IsValid())
+			if (!bIsUObject && Field != nullptr)
 			{
-				return CastField<T>(Property.Get());
+				return CastField<T>(Field);
 			}
 		}
 		return nullptr;
@@ -134,16 +134,16 @@ public:
 
 	friend uint32 GetTypeHash(const FBindingObject& BindingObject)
 	{
-		return BindingObject.bIsUObject ? GetTypeHash(BindingObject.Object.Get()) : GetTypeHash(BindingObject.Property.Get());
+		return BindingObject.bIsUObject ? GetTypeHash(BindingObject.Object) : GetTypeHash(BindingObject.Field);
 	}
 
 	bool operator==(const FBindingObject &Other) const
 	{
-		return bIsUObject == Other.bIsUObject && Object == Other.Object && Property == Other.Property;
+		return bIsUObject == Other.bIsUObject && Object == Other.Object && Field == Other.Field;
 	}
 	bool operator!=(const FBindingObject &Other) const
 	{
-		return bIsUObject != Other.bIsUObject || Object != Other.Object || Property != Other.Property;
+		return bIsUObject != Other.bIsUObject || Object != Other.Object || Field != Other.Field;
 	}
 
 	friend bool operator==(const FBindingObject &Lhs, const UObject* Rhs)
@@ -165,40 +165,49 @@ public:
 
 	friend bool operator==(const FBindingObject &Lhs, const FField* Rhs)
 	{
-		return !Lhs.IsUObject() && Lhs.Property == Rhs;
+		return !Lhs.IsUObject() && Lhs.Field == Rhs;
 	}
 	friend bool operator!=(const FBindingObject &Lhs, const FField* Rhs)
 	{
-		return Lhs.IsUObject() || Lhs.Property != Rhs;
+		return Lhs.IsUObject() || Lhs.Field != Rhs;
 	}
 	friend bool operator==(const FField* Lhs, const FBindingObject &Rhs)
 	{
-		return !Rhs.IsUObject() && Rhs.Property == Lhs;
+		return !Rhs.IsUObject() && Rhs.Field == Lhs;
 	}
 	friend bool operator!=(const FField* Lhs, const FBindingObject &Rhs)
 	{
-		return Rhs.IsUObject() || Rhs.Property != Lhs;
+		return Rhs.IsUObject() || Rhs.Field != Lhs;
 	}
 
 	friend bool operator==(const FBindingObject &Lhs, TYPE_OF_NULLPTR)
 	{
-		return Lhs.IsUObject() ? !Lhs.Object.IsValid() : !Lhs.Property.IsValid();
+		return Lhs.IsUObject() ? Lhs.Object == nullptr : Lhs.Field == nullptr;
 	}
 	friend bool operator!=(const FBindingObject &Lhs, TYPE_OF_NULLPTR)
 	{
-		return Lhs.IsUObject() ? Lhs.Object.IsValid() : Lhs.Property.IsValid();
+		return Lhs.IsUObject() ? Lhs.Object != nullptr : Lhs.Field != nullptr;
 	}
 	friend bool operator==(TYPE_OF_NULLPTR, const FBindingObject &Rhs)
 	{
-		return Rhs.IsUObject() ? !Rhs.Object.IsValid() : !Rhs.Property.IsValid();
+		return Rhs.IsUObject() ? Rhs.Object == nullptr : Rhs.Field == nullptr;
 	}
 	friend bool operator!=(TYPE_OF_NULLPTR, const FBindingObject &Rhs)
 	{
-		return Rhs.IsUObject() ? Rhs.Object.IsValid() : Rhs.Property.IsValid();
+		return Rhs.IsUObject() ? Rhs.Object != nullptr : Rhs.Field != nullptr;
 	}
-	
+
+	void AddStructReferencedObjects(FReferenceCollector& Collector);
 };
 
+template<>
+struct TStructOpsTypeTraits<FBindingObject> : public TStructOpsTypeTraitsBase2<FBindingObject>
+{
+	enum
+	{
+		WithAddStructReferencedObjects = true,
+	};
+};
 
 class IBlueprintNodeBinder
 {
diff --git a/Engine/Source/Editor/BlueprintGraph/Private/BlueprintNodeBinder.cpp b/Engine/Source/Editor/BlueprintGraph/Private/BlueprintNodeBinder.cpp
new file mode 100644
--- /dev/null
+++ b/Engine/Source/Editor/BlueprintGraph/Private/BlueprintNodeBinder.cpp
@@ -0,0 +1,15 @@
+// Copyright Epic Games, Inc. All Rights Reserved.
+
+#include "BlueprintNodeBinder.h"
+
+void FBindingObject::AddStructReferencedObjects(FReferenceCollector& Collector)
+{
+	if (bIsUObject)
+	{
+		Collector.AddReferencedObject(Object);
+	}
+	else if (FField* CurrentField = Get<FField>())
+	{
+		CurrentField->AddReferencedObjects(Collector);
+	}
+}
