PR #11745: Remove unused line in `EdGraphSchema_K2.cpp`

diff --git a/Engine/Source/Editor/BlueprintGraph/Private/EdGraphSchema_K2.cpp b/Engine/Source/Editor/BlueprintGraph/Private/EdGraphSchema_K2.cpp
--- a/Engine/Source/Editor/BlueprintGraph/Private/EdGraphSchema_K2.cpp
+++ b/Engine/Source/Editor/BlueprintGraph/Private/EdGraphSchema_K2.cpp
@@ -2657,7 +2657,6 @@ TOptional<UEdGraphSchema_K2::FSearchForAutocastFunctionResults> UEdGraphSchema_K
 			return false;
 		};
 
-		const FAutocastFunctionMap& AutocastFunctionMap = FAutocastFunctionMap::Get();
 		if (!FindAndSetCastFunction(OutputPinType, InputPinType))
 		{
 			// Since single-precision float interfaces have been deprecated,
