UE-212762 - Blueprint Add Custom Event unable to resolve delegate signature inside a split struct

The pin splitting logic didn't propagate the underlying member reference for the function signature of a delegate. CreatePin doesn't handle this automatically, so we need to copy over the information manually.

diff --git a/Engine/Source/Editor/BlueprintGraph/Private/EdGraphSchema_K2.cpp b/Engine/Source/Editor/BlueprintGraph/Private/EdGraphSchema_K2.cpp
--- a/Engine/Source/Editor/BlueprintGraph/Private/EdGraphSchema_K2.cpp
+++ b/Engine/Source/Editor/BlueprintGraph/Private/EdGraphSchema_K2.cpp
@@ -7316,6 +7316,13 @@ void UEdGraphSchema_K2::SplitPin(UEdGraphPin* Pin, const bool bNotify) const
 			PinParams.ContainerType = ProtoPinType.ContainerType;
 			PinParams.ValueTerminalType = ProtoPinType.PinValueType;
 			UEdGraphPin* SubPin = GraphNode->CreatePin(Pin->Direction, ProtoPinType.PinCategory, ProtoPinType.PinSubCategory, ProtoPinType.PinSubCategoryObject.Get(), PinName, PinParams);
+			check(SubPin);
+
+			// Delegate pins will also need a signature copied over. CreatePin doesn't handle this.
+			if (const UFunction* PinSignature = FMemberReference::ResolveSimpleMemberReference<UFunction>(ProtoPinType.PinSubCategoryMemberReference))
+			{
+				FMemberReference::FillSimpleMemberReference(PinSignature, SubPin->PinType.PinSubCategoryMemberReference);
+			}
 
 			if (K2Node != nullptr && K2Node->ShouldDrawCompact() && !Pin->ParentPin)
 			{
