Only set the new operation name upon node conversion if we successfully find a matching function to that new conversion.

#rnx

diff --git a/Engine/Source/Editor/BlueprintGraph/Private/K2Node_PromotableOperator.cpp b/Engine/Source/Editor/BlueprintGraph/Private/K2Node_PromotableOperator.cpp
--- a/Engine/Source/Editor/BlueprintGraph/Private/K2Node_PromotableOperator.cpp
+++ b/Engine/Source/Editor/BlueprintGraph/Private/K2Node_PromotableOperator.cpp
@@ -338,17 +338,6 @@ void UK2Node_PromotableOperator::ConvertComparisonOperatorNode(UEdGraphNode* Nod
 	{
 		return;
 	}
-	
-	FFormatNamedArguments Args;
-	Args.Add(TEXT("NewOpName"), FText::FromName(NewOpName));
-	Args.Add(TEXT("CurrentOpName"), FText::FromName(OpNode->OperationName));
-
-	const FText PinConversionName = FText::Format(LOCTEXT("CallFunction_Tooltip", "Convert operator node from {CurrentOpName} to {NewOpName}"), Args);
-	
-	FScopedTransaction Transaction(LOCTEXT("PromotableOperatorComparisonOpConversion", "Convert operator node"));
-	OpNode->Modify();
-	
-	OpNode->OperationName = NewOpName;
 
 	TArray<UEdGraphPin*> PinsToConsider;
 	OpNode->GetPinsToConsider(PinsToConsider);
@@ -365,8 +354,13 @@ void UK2Node_PromotableOperator::ConvertComparisonOperatorNode(UEdGraphNode* Nod
 	}
 
 	// For nodes with connections we have to find the best function again that matches for them
-	if (const UFunction* BestMatchingFunc = FTypePromotion::FindBestMatchingFunc(OpNode->OperationName, PinsToConsider))
+	if (const UFunction* BestMatchingFunc = FTypePromotion::FindBestMatchingFunc(NewOpName, PinsToConsider))
 	{
+		FScopedTransaction Transaction(LOCTEXT("PromotableOperatorComparisonOpConversion", "Convert operator node"));
+		OpNode->Modify();
+		
+		OpNode->OperationName = NewOpName;
+		
 		// Only allow this with comparison functions
 		ensure(FTypePromotion::IsComparisonFunc(BestMatchingFunc));
 		UE_LOG(LogBlueprint, Verbose, TEXT("Converting node '%s' from '%s' to '%s'..."), *GetNameSafe(OpNode), *OpNode->OperationName.ToString(), *NewOpName.ToString());
