- MakeStruct node will correctly use DisplayName meta data for the Title and Tooltip Text
- BreakStruct node will correctly use DisplayName meta data for the Title and Tooltip Text
- StructViewer respects struct's hierarchical ShowTooltip metadata

diff --git a/Engine/Source/Editor/BlueprintGraph/Private/K2Node_BreakStruct.cpp b/Engine/Source/Editor/BlueprintGraph/Private/K2Node_BreakStruct.cpp
--- a/Engine/Source/Editor/BlueprintGraph/Private/K2Node_BreakStruct.cpp
+++ b/Engine/Source/Editor/BlueprintGraph/Private/K2Node_BreakStruct.cpp
@@ -253,7 +253,7 @@ FText UK2Node_BreakStruct::GetNodeTitle(ENodeTitleType::Type TitleType) const
 	else if (CachedNodeTitle.IsOutOfDate(this))
 	{
 		FFormatNamedArguments Args;
-		Args.Add(TEXT("StructName"), FText::FromString(StructType->GetName()));
+		Args.Add(TEXT("StructName"), StructType->GetDisplayNameText());
 
 		// FText::Format() is slow, so we cache this to save on performance
 		CachedNodeTitle.SetCachedText(FText::Format(LOCTEXT("BreakNodeTitle", "Break {StructName}"), Args), this);
@@ -272,7 +272,7 @@ FText UK2Node_BreakStruct::GetTooltipText() const
 		// FText::Format() is slow, so we cache this to save on performance
 		CachedTooltip.SetCachedText(FText::Format(
 			LOCTEXT("BreakStruct_Tooltip", "Adds a node that breaks a '{0}' into its member fields"),
-			FText::FromName(StructType->GetFName())
+			StructType->GetDisplayNameText()
 		), this);
 	}
 	return CachedTooltip;
diff --git a/Engine/Source/Editor/BlueprintGraph/Private/K2Node_MakeStruct.cpp b/Engine/Source/Editor/BlueprintGraph/Private/K2Node_MakeStruct.cpp
--- a/Engine/Source/Editor/BlueprintGraph/Private/K2Node_MakeStruct.cpp
+++ b/Engine/Source/Editor/BlueprintGraph/Private/K2Node_MakeStruct.cpp
@@ -242,7 +242,7 @@ FText UK2Node_MakeStruct::GetNodeTitle(ENodeTitleType::Type TitleType) const
 	else if (CachedNodeTitle.IsOutOfDate(this))
 	{
 		FFormatNamedArguments Args;
-		Args.Add(TEXT("StructName"), FText::FromName(StructType->GetFName()));
+		Args.Add(TEXT("StructName"), StructType->GetDisplayNameText());
 		// FText::Format() is slow, so we cache this to save on performance
 		CachedNodeTitle.SetCachedText(FText::Format(LOCTEXT("MakeNodeTitle", "Make {StructName}"), Args), this);
 	}
@@ -260,7 +260,7 @@ FText UK2Node_MakeStruct::GetTooltipText() const
 		// FText::Format() is slow, so we cache this to save on performance
 		CachedTooltip.SetCachedText(FText::Format(
 			LOCTEXT("MakeStruct_Tooltip", "Adds a node that create a '{0}' from its members"),
-			FText::FromName(StructType->GetFName())
+			StructType->GetDisplayNameText()
 		), this);
 	}
 	return CachedTooltip;
diff --git a/Engine/Source/Editor/StructViewer/Private/SStructViewer.cpp b/Engine/Source/Editor/StructViewer/Private/SStructViewer.cpp
--- a/Engine/Source/Editor/StructViewer/Private/SStructViewer.cpp
+++ b/Engine/Source/Editor/StructViewer/Private/SStructViewer.cpp
@@ -643,7 +643,7 @@ public:
 			else if (!AssociatedNode->GetStructPath().IsNull())
 			{
 				const UScriptStruct* Struct = AssociatedNode->GetStruct();
-				if (Struct != nullptr && Struct->GetBoolMetaData("ShowTooltip"))
+				if (Struct != nullptr && Struct->GetBoolMetaDataHierarchical("ShowTooltip"))
 				{
 					const FText ToolTipText = FText::Format(LOCTEXT("ToolTipFormat", "{0}\n\n{1}"), 
 						AssociatedNode->GetStruct()->GetToolTipText(), 
