Remove redundant hidden category information from blueprint context menu - optimizing category based hiding when redundant categories are hidden
#rnx

diff --git a/Engine/Source/Editor/BlueprintGraph/Private/BlueprintActionFilter.cpp b/Engine/Source/Editor/BlueprintGraph/Private/BlueprintActionFilter.cpp
--- a/Engine/Source/Editor/BlueprintGraph/Private/BlueprintActionFilter.cpp
+++ b/Engine/Source/Editor/BlueprintGraph/Private/BlueprintActionFilter.cpp
@@ -1,6 +1,7 @@
 // Copyright Epic Games, Inc. All Rights Reserved.
 
 #include "BlueprintActionFilter.h"
+#include "Algo/Unique.h"
 #include "UObject/Interface.h"
 #include "EdGraph/EdGraphPin.h"
 #include "Engine/Blueprint.h"
@@ -2264,6 +2265,8 @@ void FBlueprintActionFilter::Add(TArray<FTargetClassFilterData>& ToArray, UClass
 {
 	TArray<FString> ClassHideCategories;
 	FEditorCategoryUtils::GetClassHideCategories(TargetClass, ClassHideCategories);
+	ClassHideCategories.Sort();
+	ClassHideCategories.SetNum(Algo::Unique(ClassHideCategories));
 	FTargetClassFilterData Data = { TargetClass, MoveTemp(ClassHideCategories) };
 	ToArray.Add(MoveTemp(Data));
 }
