Removed display modes for "real" pin categories.

CL 19140882 added the ability to switch displays for "real" types. However, the de facto display is either "Float (single-precision)" or "Float (double-precision)". The other display modes simply confuse things.

#rnx

diff --git a/Engine/Source/Editor/BlueprintGraph/Private/EdGraphSchema_K2.cpp b/Engine/Source/Editor/BlueprintGraph/Private/EdGraphSchema_K2.cpp
--- a/Engine/Source/Editor/BlueprintGraph/Private/EdGraphSchema_K2.cpp
+++ b/Engine/Source/Editor/BlueprintGraph/Private/EdGraphSchema_K2.cpp
@@ -114,22 +114,6 @@ enum class EBlueprintRealDisplayMode : uint8
 
 namespace UE::EdGraphSchemaK2::Private
 {
-	int32 LastRealNamingMode = -1;
-	int32 RealNamingMode = static_cast<int32>(EBlueprintRealDisplayMode::Float);
-	FAutoConsoleVariableRef CVarRealNamingMode(TEXT("Blueprint.PC_Real.DisplayMode"), RealNamingMode, TEXT("Real naming mode\n\t0: Real\n\t1: Float (default)\n\t2: Number\n\nNote the editor needs to be restarted for this to fully take effect"));
-
-	EBlueprintRealDisplayMode GetRealDisplayMode()
-	{
-		return static_cast<EBlueprintRealDisplayMode>(FMath::Clamp(RealNamingMode, 0, 2));
-	}
-
-	bool ShouldRefreshRealDisplay()
-	{
-		const bool bResult = LastRealNamingMode != RealNamingMode;
-		LastRealNamingMode = RealNamingMode;
-		return bResult;
-	}
-
 	template <class... T>
 	constexpr bool TAlwaysFalse = false;
 
@@ -3823,9 +3807,9 @@ FText UEdGraphSchema_K2::GetCategoryText(FName Category, FName SubCategory, bool
 		CategoryDescriptions.Add(PC_Class, LOCTEXT("ClassCategory", "Class Reference"));
 		CategoryDescriptions.Add(PC_Int, LOCTEXT("IntCategory", "Integer"));
 		CategoryDescriptions.Add(PC_Int64, LOCTEXT("Int64Category", "Integer64"));
-		CategoryDescriptions.Add(PC_Real, LOCTEXT("RealCategory", "Real"));
-		CategoryDescriptions.Add(PC_Float, LOCTEXT("FloatCategory", "Real (single-precision)"));
-		CategoryDescriptions.Add(PC_Double, LOCTEXT("DoubleCategory", "Real (double-precision)"));
+		CategoryDescriptions.Add(PC_Real, LOCTEXT("RealCategory", "Float"));
+		CategoryDescriptions.Add(PC_Float, LOCTEXT("FloatCategory", "Float (single-precision)"));
+		CategoryDescriptions.Add(PC_Double, LOCTEXT("DoubleCategory", "Float (double-precision)"));
 		CategoryDescriptions.Add(PC_Name, LOCTEXT("NameCategory", "Name"));
 		CategoryDescriptions.Add(PC_Delegate, LOCTEXT("DelegateCategory", "Delegate"));
 		CategoryDescriptions.Add(PC_MCDelegate, LOCTEXT("MulticastDelegateCategory", "Multicast Delegate"));
@@ -3842,31 +3826,6 @@ FText UEdGraphSchema_K2::GetCategoryText(FName Category, FName SubCategory, bool
 		CategoryDescriptions.Add(AllObjectTypes, LOCTEXT("AllObjectTypes", "Object Types"));
 	}
 
-	if (ShouldRefreshRealDisplay())
-	{
-		switch (GetRealDisplayMode())
-		{
-		case EBlueprintRealDisplayMode::Real:
-			CategoryDescriptions[PC_Real] = LOCTEXT("RealCategory_DisplayAsReal", "Real");
-			CategoryDescriptions[PC_Float] = LOCTEXT("RealCategory_DisplayAsReal_SinglePrecision", "Real (single-precision)");
-			CategoryDescriptions[PC_Double] = LOCTEXT("RealCategory_DisplayAsReal_DoublePrecision", "Real (double-precision)");
-			break;
-		case EBlueprintRealDisplayMode::Float:
-			CategoryDescriptions[PC_Real] = LOCTEXT("RealCategory_DisplayAsFloat", "Float");
-			CategoryDescriptions[PC_Float] = LOCTEXT("RealCategory_DisplayAsFloat_SinglePrecision", "Float (single-precision)");
-			CategoryDescriptions[PC_Double] = LOCTEXT("RealCategory_DisplayAsFloat_DoublePrecision", "Float (double-precision)");
-			break;
-		case EBlueprintRealDisplayMode::Number:
-			CategoryDescriptions[PC_Real] = LOCTEXT("RealCategory_DisplayAsNumber", "Number");
-			CategoryDescriptions[PC_Float] = LOCTEXT("RealCategory_DisplayAsNumber_SinglePrecision", "Number (single-precision)");
-			CategoryDescriptions[PC_Double] = LOCTEXT("RealCategory_DisplayAsNumber_DoublePrecision", "Number (double-precision)");
-			break;
-		default:
-			check(false);
-			break;
-		}
-	}
-
 	if (const FText* TypeDesc = CategoryDescriptions.Find(Category))
 	{
 		const bool bUseDetailedRealCategory =
@@ -3950,20 +3909,7 @@ FText UEdGraphSchema_K2::TerminalTypeToText(const FName Category, const FName Su
 	{
 		if (Category == UEdGraphSchema_K2::PC_Real)
 		{
-			using namespace UE::EdGraphSchemaK2::Private;
-
-			switch (GetRealDisplayMode())
-			{
-			case EBlueprintRealDisplayMode::Real:
-				PropertyText = (SubCategory == UEdGraphSchema_K2::PC_Float) ? LOCTEXT("SinglePrecisionReal", "Real (single-precision)") : LOCTEXT("DoublePrecisionReal", "Real (double-precision)");
-				break;
-			case EBlueprintRealDisplayMode::Float:
-				PropertyText = (SubCategory == UEdGraphSchema_K2::PC_Float) ? LOCTEXT("SinglePrecisionFloat", "Float (single-precision)") : LOCTEXT("DoublePrecisionFloat", "Float (double-precision)");
-				break;
-			case EBlueprintRealDisplayMode::Number:
-				PropertyText = (SubCategory == UEdGraphSchema_K2::PC_Float) ? LOCTEXT("SinglePrecisionNumber", "Number (single-precision)") : LOCTEXT("DoublePrecisionNumber", "Number (double-precision)");
-				break;
-			}
+			PropertyText = UEdGraphSchema_K2::GetCategoryText(Category, SubCategory, true);
 		}
 		else
 		{
@@ -4027,13 +3973,13 @@ void UEdGraphSchema_K2::GetVariableTypeTree(TArray< TSharedPtr<FPinTypeTreeInfo>
 
 	if( bAllowExec )
 	{
-		TypeTree.Add( MakeShareable( new FPinTypeTreeInfo(GetCategoryText(PC_Exec, true), PC_Exec, this, LOCTEXT("ExecType", "Execution pin")) ) );
+		TypeTree.Add(MakeShareable(new FPinTypeTreeInfo(GetCategoryText(PC_Exec, true), PC_Exec, this, LOCTEXT("ExecType", "Execution pin"))));
 	}
 
-	TypeTree.Add( MakeShareable( new FPinTypeTreeInfo(GetCategoryText(PC_Boolean, true), PC_Boolean, this, LOCTEXT("BooleanType", "True or false value")) ) );
-	TypeTree.Add( MakeShareable( new FPinTypeTreeInfo(GetCategoryText(PC_Byte, true), PC_Byte, this, LOCTEXT("ByteType", "8 bit number")) ) );
-	TypeTree.Add( MakeShareable( new FPinTypeTreeInfo(GetCategoryText(PC_Int, true), PC_Int, this, LOCTEXT("IntegerType", "Integer number")) ) );
-	TypeTree.Add( MakeShareable( new FPinTypeTreeInfo(GetCategoryText(PC_Int64, true), PC_Int64, this, LOCTEXT("Integer64Type", "64 bit Integer number")) ) );
+	TypeTree.Add(MakeShareable(new FPinTypeTreeInfo(GetCategoryText(PC_Boolean, true), PC_Boolean, this, LOCTEXT("BooleanType", "True or false value"))));
+	TypeTree.Add(MakeShareable(new FPinTypeTreeInfo(GetCategoryText(PC_Byte, true), PC_Byte, this, LOCTEXT("ByteType", "8 bit number"))));
+	TypeTree.Add(MakeShareable(new FPinTypeTreeInfo(GetCategoryText(PC_Int, true), PC_Int, this, LOCTEXT("IntegerType", "Integer number"))));
+	TypeTree.Add(MakeShareable(new FPinTypeTreeInfo(GetCategoryText(PC_Int64, true), PC_Int64, this, LOCTEXT("Integer64Type", "64 bit Integer number"))));
 
 	if (!bIndexTypesOnly)
 	{
@@ -4053,7 +3999,7 @@ void UEdGraphSchema_K2::GetVariableTypeTree(TArray< TSharedPtr<FPinTypeTreeInfo>
 	// Add wildcard type
 	if (bAllowWildCard)
 	{
-		TypeTree.Add( MakeShareable( new FPinTypeTreeInfo(GetCategoryText(PC_Wildcard, true), PC_Wildcard, this, LOCTEXT("WildcardType", "Wildcard type (unspecified)")) ) );
+		TypeTree.Add(MakeShareable(new FPinTypeTreeInfo(GetCategoryText(PC_Wildcard, true), PC_Wildcard, this, LOCTEXT("WildcardType", "Wildcard type (unspecified)"))));
 	}
 
 	// Add the types that have subtrees
