Changed instances of FText::FromString in editor graph nodes to use AsCultureInvariant. This is required in case the nodes are cooked out (can occur in the case of optional editor data). These strings do not need localization and will trigger warnings in the cook if they aren't marked invariant.

diff --git a/Engine/Source/Editor/BlueprintGraph/Private/K2Node.cpp b/Engine/Source/Editor/BlueprintGraph/Private/K2Node.cpp
--- a/Engine/Source/Editor/BlueprintGraph/Private/K2Node.cpp
+++ b/Engine/Source/Editor/BlueprintGraph/Private/K2Node.cpp
@@ -1691,7 +1691,7 @@ ERenamePinResult UK2Node::RenameUserDefinedPinImpl(const FName OldName, const FN
 				FString NewNameStr = NewName.ToString();
 
 				PinToRename->PinName = *(NewNameStr + PinToRename->PinName.ToString().RightChop(OldNameLength));
-				PinToRename->PinFriendlyName = FText::FromString(MoveTemp(NewNameStr) + PinToRename->PinFriendlyName.ToString().RightChop(OldNameLength));
+				PinToRename->PinFriendlyName = FText::AsCultureInvariant(MoveTemp(NewNameStr) + PinToRename->PinFriendlyName.ToString().RightChop(OldNameLength));
 			}
 		}
 	}
@@ -1839,7 +1839,7 @@ void FOptionalPinManager::CreateVisiblePins(TArray<FOptionalPinFromProperty>& Pr
 							const FName PinName = *FString::Printf(TEXT("%s_%d"), *PropertyEntry.PropertyName.ToString(), Index);
 
 							FFormatNamedArguments Args;
-							Args.Add(TEXT("PinName"), FText::FromString(PropertyEntry.PropertyFriendlyName.IsEmpty() ? PinName.ToString() : PropertyEntry.PropertyFriendlyName));
+							Args.Add(TEXT("PinName"), FText::AsCultureInvariant(PropertyEntry.PropertyFriendlyName.IsEmpty() ? PinName.ToString() : PropertyEntry.PropertyFriendlyName));
 							Args.Add(TEXT("Index"), Index);
 							const FText PinFriendlyName = FText::Format(LOCTEXT("PinFriendlyNameWithIndex", "{PinName} {Index}"), Args);
 
@@ -1890,7 +1890,7 @@ void FOptionalPinManager::CreateVisiblePins(TArray<FOptionalPinFromProperty>& Pr
 					{
 						const FName PinName = PropertyEntry.PropertyName;
 						NewPin = TargetNode->CreatePin(Direction, PinType, PinName);
-						NewPin->PinFriendlyName = FText::FromString(PropertyEntry.PropertyFriendlyName.IsEmpty() ? PinName.ToString() : PropertyEntry.PropertyFriendlyName);
+						NewPin->PinFriendlyName = FText::AsCultureInvariant(PropertyEntry.PropertyFriendlyName.IsEmpty() ? PinName.ToString() : PropertyEntry.PropertyFriendlyName);
 						NewPin->bNotConnectable = !PropertyEntry.bIsSetValuePinVisible;
 						NewPin->bDefaultValueIsIgnored = !PropertyEntry.bIsSetValuePinVisible;
 						Schema->ConstructBasicPinTooltip(*NewPin, PropertyEntry.PropertyTooltip, NewPin->PinToolTip);
