Fixed another instance of a blueprint graph pin taking it's name from FText::FromString, causing cook time warnings due to missing localization.

diff --git a/Engine/Source/Editor/BlueprintGraph/Private/EdGraphSchema_K2.cpp b/Engine/Source/Editor/BlueprintGraph/Private/EdGraphSchema_K2.cpp
--- a/Engine/Source/Editor/BlueprintGraph/Private/EdGraphSchema_K2.cpp
+++ b/Engine/Source/Editor/BlueprintGraph/Private/EdGraphSchema_K2.cpp
@@ -3178,7 +3178,7 @@ FText UEdGraphSchema_K2::GetPinDisplayName(const UEdGraphPin* Pin) const
 
 		if( GEditor && GetDefault<UEditorStyleSettings>()->bShowFriendlyNames && Pin->bAllowFriendlyName )
 		{
-			DisplayName = FText::FromString(FName::NameToDisplayString(DisplayName.ToString(), Pin->PinType.PinCategory == PC_Boolean));
+			DisplayName = FText::AsCultureInvariant(FName::NameToDisplayString(DisplayName.ToString(), Pin->PinType.PinCategory == PC_Boolean));
 		}
 	}
 	return DisplayName;
diff --git a/Engine/Source/Editor/BlueprintGraph/Private/K2Node_CallFunction.cpp b/Engine/Source/Editor/BlueprintGraph/Private/K2Node_CallFunction.cpp
--- a/Engine/Source/Editor/BlueprintGraph/Private/K2Node_CallFunction.cpp
+++ b/Engine/Source/Editor/BlueprintGraph/Private/K2Node_CallFunction.cpp
@@ -938,7 +938,7 @@ void UK2Node_CallFunction::CreateExecPinsForFunctionCall(const UFunction* Functi
 								else
 								{
 									CreatedPin = CreatePin(Direction, UEdGraphSchema_K2::PC_Exec, *NameStr);
-									CreatedPin->PinFriendlyName = FText::FromString(FString::Printf(TEXT("(%s) %s"), *Prop->GetDisplayNameText().ToString(), *NameStr));
+									CreatedPin->PinFriendlyName = FText::AsCultureInvariant(FString::Printf(TEXT("(%s) %s"), *Prop->GetDisplayNameText().ToString(), *NameStr));
 								}
 
 								ExpandAsEnumPins.Add(CreatedPin);
@@ -1201,7 +1201,7 @@ bool UK2Node_CallFunction::CreatePinsForFunctionCall(const UFunction* Function)
 			const FString& PinDisplayName = Param->GetMetaData(FBlueprintMetadata::MD_DisplayName);
 			if (!PinDisplayName.IsEmpty())
 			{
-				Pin->PinFriendlyName = FText::FromString(PinDisplayName);
+				Pin->PinFriendlyName = FText::AsCultureInvariant(PinDisplayName);
 			}
 			else if (Function->GetReturnProperty() == Param && Function->HasMetaData(FBlueprintMetadata::MD_ReturnDisplayName))
 			{
