K2Node_Variable is more careful abound member reference usage to avoid resetting member name due to a stale guid. K2Node_FunctionEntry now generates a fresh local guid when it is duplicated. This fixes a phantom compile error after pasting local variable nodes, which resulted in NodeGuid collisions

[RN] Fixed one cause of phantom blueprint compilation errors after duplicating a blueprint graph

diff --git a/Engine/Source/Editor/BlueprintGraph/Private/K2Node_FunctionEntry.cpp b/Engine/Source/Editor/BlueprintGraph/Private/K2Node_FunctionEntry.cpp
--- a/Engine/Source/Editor/BlueprintGraph/Private/K2Node_FunctionEntry.cpp
+++ b/Engine/Source/Editor/BlueprintGraph/Private/K2Node_FunctionEntry.cpp
@@ -884,6 +884,11 @@ void UK2Node_FunctionEntry::PostPasteNode()
 	// If a function entry is being pasted, it should be editable in it's new graph
 	bIsEditable = true;
 
+	for (FBPVariableDescription& LocalVariable : LocalVariables)
+	{
+		LocalVariable.VarGuid = FGuid::NewGuid();
+	}
+
 	// ensure there are UserDefinedPins for all pins except the 'then' pin
 	for (int32 PinIdx = 1; PinIdx < Pins.Num(); ++PinIdx)
 	{
diff --git a/Engine/Source/Editor/BlueprintGraph/Private/K2Node_Variable.cpp b/Engine/Source/Editor/BlueprintGraph/Private/K2Node_Variable.cpp
--- a/Engine/Source/Editor/BlueprintGraph/Private/K2Node_Variable.cpp
+++ b/Engine/Source/Editor/BlueprintGraph/Private/K2Node_Variable.cpp
@@ -93,6 +93,21 @@ void UK2Node_Variable::SetFromProperty(const FProperty* Property, bool bSelfCont
 bool UK2Node_Variable::CreatePinForVariable(EEdGraphPinDirection Direction, FName InPinName/* = NAME_None */)
 {
 	const UEdGraphSchema_K2* K2Schema = GetDefault<UEdGraphSchema_K2>();
+	// This must run before GetPropertyForVariable to avoid guid collisions in pastebins
+	// from causing the variable reference to be renamed by FMemberReference::ResolveMemberImpl
+	// via FMemberReference::RefreshLocalVariableName:
+	if (VariableReference.IsLocalScope())
+	{
+		if (FBPVariableDescription const* Desc = GetBlueprintVarDescription())
+		{
+			const FName PinName = InPinName.IsNone() ? GetVarName() : InPinName;
+			UEdGraphPin* VariablePin = CreatePin(Direction, NAME_None, PinName);
+			// search for local declaration:
+			VariablePin->PinType = Desc->VarType;
+			K2Schema->SetPinAutogeneratedDefaultValueBasedOnType(VariablePin);
+			return true;
+		}
+	}
 
 	FProperty* VariableProperty = GetPropertyForVariable();
 	// favor the skeleton property if possible (in case the property type has 
@@ -1031,10 +1046,18 @@ FBPVariableDescription const* UK2Node_Variable::GetBlueprintVarDescription() con
 	FName const& VarName = VariableReference.GetMemberName();
 	UStruct const* VariableScope = VariableReference.GetMemberScope(GetBlueprintClassFromNode());
 
-	bool const bIsLocalVariable = (VariableScope != nullptr);
+	bool const bIsLocalVariable = VariableReference.IsLocalScope();
 	if (bIsLocalVariable)
 	{
-		return FBlueprintEditorUtils::FindLocalVariable(GetBlueprint(), VariableScope, VarName);
+		if (VariableScope)
+		{
+			if (FBPVariableDescription const* Local = FBlueprintEditorUtils::FindLocalVariable(GetBlueprint(), VariableScope, VarName))
+			{
+				return Local;
+			}
+		}
+
+		return FBlueprintEditorUtils::FindLocalVariable(GetBlueprint(), GetGraph(), VarName);
 	}
 	else if (FProperty const* VarProperty = GetPropertyForVariable())
 	{
@@ -1075,11 +1098,7 @@ void UK2Node_Variable::PostPasteNode()
 	UBlueprint* Blueprint = GetBlueprint();
 	bool bInvalidateVariable = false;
 
-	if (VariableReference.ResolveMember<FProperty>(Blueprint) == nullptr)
-	{
-		bInvalidateVariable = true;
-	}
-	else if (VariableReference.IsLocalScope())
+	if (VariableReference.IsLocalScope())
 	{
 		// Local scoped variables should always validate whether they are being placed in the same graph as their scope
 		// ResolveMember will not return nullptr when the graph changes but the Blueprint remains the same.
@@ -1091,6 +1110,11 @@ void UK2Node_Variable::PostPasteNode()
 			bInvalidateVariable = true;
 		}
 	}
+
+	if(!bInvalidateVariable && VariableReference.ResolveMember<FProperty>(Blueprint) == nullptr)
+	{
+		bInvalidateVariable = true;
+	}
 		
 	if (bInvalidateVariable)
 	{
