Fixed SetFieldsInStruct node to correctly display the DisplayName of the struct in the title and tooltip

diff --git a/Engine/Source/Editor/BlueprintGraph/Private/K2Node_SetFieldsInStruct.cpp b/Engine/Source/Editor/BlueprintGraph/Private/K2Node_SetFieldsInStruct.cpp
--- a/Engine/Source/Editor/BlueprintGraph/Private/K2Node_SetFieldsInStruct.cpp
+++ b/Engine/Source/Editor/BlueprintGraph/Private/K2Node_SetFieldsInStruct.cpp
@@ -193,7 +193,7 @@ FText UK2Node_SetFieldsInStruct::GetNodeTitle(ENodeTitleType::Type TitleType) co
 	else if (CachedNodeTitle.IsOutOfDate(this))
 	{
 		FFormatNamedArguments Args;
-		Args.Add(TEXT("StructName"), FText::FromName(StructType->GetFName()));
+		Args.Add(TEXT("StructName"), StructType->GetDisplayNameText());
 		// FText::Format() is slow, so we cache this to save on performance
 		CachedNodeTitle.SetCachedText(FText::Format(LOCTEXT("SetFieldsInStructNodeTitle", "Set members in {StructName}"), Args), this);
 	}
@@ -211,7 +211,7 @@ FText UK2Node_SetFieldsInStruct::GetTooltipText() const
 		// FText::Format() is slow, so we cache this to save on performance
 		CachedTooltip.SetCachedText(FText::Format(
 			LOCTEXT("SetFieldsInStruct_Tooltip", "Adds a node that modifies a '{0}'"),
-			FText::FromName(StructType->GetFName())
+			StructType->GetDisplayNameText()
 		), this);
 	}
 	return CachedTooltip;
