Add generic ToText node that uses the local culture.

diff --git a/Engine/Plugins/Runtime/ModelViewViewModel/Source/ModelViewViewModelBlueprint/Private/MVVMDeveloperProjectSettings.cpp b/Engine/Plugins/Runtime/ModelViewViewModel/Source/ModelViewViewModelBlueprint/Private/MVVMDeveloperProjectSettings.cpp
--- a/Engine/Plugins/Runtime/ModelViewViewModel/Source/ModelViewViewModelBlueprint/Private/MVVMDeveloperProjectSettings.cpp
+++ b/Engine/Plugins/Runtime/ModelViewViewModel/Source/ModelViewViewModelBlueprint/Private/MVVMDeveloperProjectSettings.cpp
@@ -38,8 +38,10 @@ UMVVMDeveloperProjectSettings::UMVVMDeveloperProjectSettings()
 
 	FTopLevelAssetPath BlueprintFunctionLibrary = FTopLevelAssetPath("/Script/Engine", "BlueprintFunctionLibrary");
 	FTopLevelAssetPath FormatText = FTopLevelAssetPath("/Script/BlueprintGraph", "K2Node_FormatText");
+	FTopLevelAssetPath GenericToText = FTopLevelAssetPath("/Script/BlueprintGraph", "K2Node_GenericToText");
 	AllowedClassForConversionFunctions.Add(FSoftClassPath(BlueprintFunctionLibrary.ToString()));
 	AllowedClassForConversionFunctions.Add(FSoftClassPath(FormatText.ToString()));
+	AllowedClassForConversionFunctions.Add(FSoftClassPath(GenericToText.ToString()));
 
 	SupportedListViewBaseClassesForExtension.Add(UListView::StaticClass());
 	SupportedPanelClassesForExtension.Add(UHorizontalBox::StaticClass());
diff --git a/Engine/Source/Editor/BlueprintGraph/Classes/K2Node_GenericToText.h b/Engine/Source/Editor/BlueprintGraph/Classes/K2Node_GenericToText.h
new file mode 100644
--- /dev/null
+++ b/Engine/Source/Editor/BlueprintGraph/Classes/K2Node_GenericToText.h
@@ -0,0 +1,42 @@
+// Copyright Epic Games, Inc. All Rights Reserved.
+
+#pragma once
+
+#include "CoreMinimal.h"
+#include "K2Node.h"
+
+#include "K2Node_GenericToText.generated.h"
+
+
+UCLASS(MinimalAPI)
+class UK2Node_GenericToText : public UK2Node
+{
+	GENERATED_BODY()
+
+	//~ Begin UK2Node Interface.
+	virtual void AllocateDefaultPins() override;
+	virtual void PinConnectionListChanged(UEdGraphPin* Pin) override;
+	virtual void PinTypeChanged(UEdGraphPin* Pin) override;
+	virtual void NodeConnectionListChanged() override;
+	virtual bool IsConnectionDisallowed(const UEdGraphPin* MyPin, const UEdGraphPin* OtherPin, FString& OutReason) const override;
+	virtual void PostReconstructNode() override;
+	virtual void ExpandNode(class FKismetCompilerContext& CompilerContext, UEdGraph* SourceGraph) override;
+	virtual bool IsNodePure() const override
+	{
+		return true;
+	}
+	virtual bool NodeCausesStructuralBlueprintChange() const override
+	{
+		return true;
+	}
+
+	virtual FText GetTooltipText() const override;
+	virtual FText GetNodeTitle(ENodeTitleType::Type TitleType) const override;
+	virtual void GetMenuActions(FBlueprintActionDatabaseRegistrar& ActionRegistrar) const override;
+	virtual FText GetMenuCategory() const override;
+	//~ End UK2Node Interface.
+
+private:
+	/** Synchronize the type of the given argument pin with the type its connected to, or reset it to a wildcard pin if there's no connection */
+	void SynchronizeArgumentPinType();
+};
diff --git a/Engine/Source/Editor/BlueprintGraph/Private/K2Node_GenericToText.cpp b/Engine/Source/Editor/BlueprintGraph/Private/K2Node_GenericToText.cpp
new file mode 100644
--- /dev/null
+++ b/Engine/Source/Editor/BlueprintGraph/Private/K2Node_GenericToText.cpp
@@ -0,0 +1,289 @@
+// Copyright Epic Games, Inc. All Rights Reserved.
+
+
+#include "K2Node_GenericToText.h"
+
+#include "BlueprintActionDatabaseRegistrar.h"
+#include "BlueprintNodeSpawner.h"
+#include "EdGraph/EdGraph.h"
+#include "EdGraphSchema_K2.h"
+#include "EditorCategoryUtils.h"
+#include "K2Node_CallFunction.h"
+#include "Kismet/KismetTextLibrary.h"
+#include "Kismet2/BlueprintEditorUtils.h"
+#include "Kismet2/CompilerResultsLog.h"
+#include "KismetCompiler.h"
+
+#define LOCTEXT_NAMESPACE "K2Node_GenericToText"
+
+/////////////////////////////////////////////////////
+// K2Node_GenericConvToText
+
+namespace UE::Blueprint::Private
+{
+	const FLazyName PC_Value = "Value";
+	const TArrayView<const UScriptStruct*> GetToTextSupportedScriptStruct()
+	{
+		static const UScriptStruct* List[] = { TBaseStructure<FVector>::Get()
+			, TBaseStructure<FVector2D>::Get()
+			, TBaseStructure<FRotator>::Get()
+			, TBaseStructure<FTransform>::Get()
+			, TBaseStructure<FLinearColor>::Get()
+			, TBaseStructure<FDateTime>::Get()
+		};
+		return MakeArrayView(List);
+	}
+
+	bool IsGenericNumericProperty(FName OtherPinCategory)
+	{
+		return OtherPinCategory == UEdGraphSchema_K2::PC_Boolean
+			|| OtherPinCategory == UEdGraphSchema_K2::PC_Byte
+			|| OtherPinCategory == UEdGraphSchema_K2::PC_Enum
+			|| OtherPinCategory == UEdGraphSchema_K2::PC_Int
+			|| OtherPinCategory == UEdGraphSchema_K2::PC_Int64
+			|| OtherPinCategory == UEdGraphSchema_K2::PC_Real
+			|| OtherPinCategory == UEdGraphSchema_K2::PC_Double
+			|| OtherPinCategory == UEdGraphSchema_K2::PC_Float;
+	}
+}
+
+void UK2Node_GenericToText::AllocateDefaultPins()
+{
+	Super::AllocateDefaultPins();
+
+	CreatePin(EGPD_Input, UEdGraphSchema_K2::PC_Wildcard, UE::Blueprint::Private::PC_Value.Resolve());
+	CreatePin(EGPD_Output, UEdGraphSchema_K2::PC_Text, UEdGraphSchema_K2::PN_ReturnValue);
+}
+
+void UK2Node_GenericToText::SynchronizeArgumentPinType()
+{
+	bool bPinTypeChanged = false;
+	UEdGraphPin* InputPin = FindPinChecked(UE::Blueprint::Private::PC_Value.Resolve(), EEdGraphPinDirection::EGPD_Input);
+	if (InputPin->LinkedTo.Num() == 0)
+	{
+		static const FEdGraphPinType WildcardPinType = FEdGraphPinType(UEdGraphSchema_K2::PC_Wildcard, NAME_None, nullptr, EPinContainerType::None, false, FEdGraphTerminalType());
+
+		// Ensure wildcard
+		if (InputPin->PinType != WildcardPinType)
+		{
+			InputPin->PinType = WildcardPinType;
+			bPinTypeChanged = true;
+		}
+	}
+	else
+	{
+		UEdGraphPin* SourcePin = InputPin->LinkedTo[0];
+
+		// Take the type of the connected pin
+		if (InputPin->PinType != SourcePin->PinType)
+		{
+			InputPin->PinType = SourcePin->PinType;
+			bPinTypeChanged = true;
+		}
+	}
+
+	if (bPinTypeChanged)
+	{
+		// Let the graph know to refresh
+		GetGraph()->NotifyNodeChanged(this);
+
+		UBlueprint* Blueprint = GetBlueprint();
+		if (!Blueprint->bBeingCompiled)
+		{
+			FBlueprintEditorUtils::MarkBlueprintAsModified(Blueprint);
+		}
+	}
+}
+
+FText UK2Node_GenericToText::GetNodeTitle(ENodeTitleType::Type TitleType) const
+{
+	return LOCTEXT("ToText_Title", "To Text");
+}
+
+void UK2Node_GenericToText::PinConnectionListChanged(UEdGraphPin* Pin)
+{
+	Modify();
+
+	// Potentially update an argument pin type
+	SynchronizeArgumentPinType();
+}
+
+void UK2Node_GenericToText::PinTypeChanged(UEdGraphPin* Pin)
+{
+	SynchronizeArgumentPinType();
+
+	Super::PinTypeChanged(Pin);
+}
+
+void UK2Node_GenericToText::NodeConnectionListChanged()
+{
+	SynchronizeArgumentPinType();
+
+	Super::NodeConnectionListChanged();
+}
+
+FText UK2Node_GenericToText::GetTooltipText() const
+{
+	return LOCTEXT("NodeTooltip", "Converts numeric value, enum, string, name and some basic structure to text.\n  \u2022 The basic type are Vector, Vector2D, Rotator, Transform, LinearColor, DateTime.");
+}
+
+void UK2Node_GenericToText::PostReconstructNode()
+{
+	Super::PostReconstructNode();
+
+	if (!IsTemplate())
+	{
+		// Make sure we're not dealing with a menu node
+		UEdGraph* OuterGraph = GetGraph();
+		if (OuterGraph && OuterGraph->Schema)
+		{
+			// Potentially update an argument pin type
+			SynchronizeArgumentPinType();
+		}
+	}
+}
+
+void UK2Node_GenericToText::ExpandNode(class FKismetCompilerContext& CompilerContext, UEdGraph* SourceGraph)
+{
+	Super::ExpandNode(CompilerContext, SourceGraph);
+
+	UEdGraphPin* InputPin = FindPinChecked(UE::Blueprint::Private::PC_Value.Resolve(), EEdGraphPinDirection::EGPD_Input);
+	if (InputPin == nullptr)
+	{
+		CompilerContext.MessageLog.Error(*LOCTEXT("Error_BadInputPin", "Invalid input pin.").ToString());
+		return;
+	}
+
+	if (InputPin->LinkedTo.Num() != 1 || InputPin->PinType.PinCategory == UEdGraphSchema_K2::PC_Wildcard)
+	{
+		CompilerContext.MessageLog.Error(*LOCTEXT("Error_BadLinkedInputPin", "Input pin is not linked to a value.").ToString());
+		return;
+	}
+
+	// Convert to the correct function call
+	auto CreateFunctionCall = [Self = this, &CompilerContext, SourceGraph](FName FunctionName)
+	{
+		UK2Node_CallFunction* CallFunction = CompilerContext.SpawnIntermediateNode<UK2Node_CallFunction>(Self, SourceGraph);
+		CallFunction->SetFromFunction(UKismetTextLibrary::StaticClass()->FindFunctionByName(FunctionName));
+		check(CallFunction->IsNodePure());
+		CallFunction->AllocateDefaultPins();
+		CompilerContext.MessageLog.NotifyIntermediateObjectCreation(CallFunction, Self);
+		return CallFunction;
+	};
+
+	UK2Node_CallFunction* CallFunction = nullptr;
+	UEdGraphPin* CallFunctionInputValue = nullptr;
+	if (UE::Blueprint::Private::IsGenericNumericProperty(InputPin->PinType.PinCategory))
+	{
+		CallFunction = CreateFunctionCall(GET_MEMBER_NAME_CHECKED(UKismetTextLibrary, Conv_NumericPropertyToText));
+		CallFunctionInputValue = CallFunction->FindPinChecked(TEXT("Value"));
+		// Set the generic type
+		CallFunctionInputValue->PinType = InputPin->PinType;
+	}
+	else if (InputPin->PinType.PinCategory == UEdGraphSchema_K2::PC_String)
+	{
+		CallFunction = CreateFunctionCall(GET_MEMBER_NAME_CHECKED(UKismetTextLibrary, Conv_StringToText));
+		CallFunctionInputValue = CallFunction->FindPinChecked(TEXT("InString"));
+	}
+	else if (InputPin->PinType.PinCategory == UEdGraphSchema_K2::PC_Name)
+	{
+		CallFunction = CreateFunctionCall(GET_MEMBER_NAME_CHECKED(UKismetTextLibrary, Conv_NameToText));
+		CallFunctionInputValue = CallFunction->FindPinChecked(TEXT("InName"));
+	}
+	else if (InputPin->PinType.PinCategory == UEdGraphSchema_K2::PC_Struct)
+	{
+		const UObject* PinSubCategoryObject = InputPin->PinType.PinSubCategoryObject.Get();
+		if (PinSubCategoryObject == TBaseStructure<FVector>::Get())
+		{
+			CallFunction = CreateFunctionCall(GET_MEMBER_NAME_CHECKED(UKismetTextLibrary, Conv_VectorToText));
+			CallFunctionInputValue = CallFunction->FindPinChecked(TEXT("InVec"));
+		}
+		else if (PinSubCategoryObject == TBaseStructure<FVector2D>::Get())
+		{
+			CallFunction = CreateFunctionCall(GET_MEMBER_NAME_CHECKED(UKismetTextLibrary, Conv_Vector2dToText));
+			CallFunctionInputValue = CallFunction->FindPinChecked(TEXT("InVec"));
+		}
+		else if (PinSubCategoryObject == TBaseStructure<FRotator>::Get())
+		{
+			CallFunction = CreateFunctionCall(GET_MEMBER_NAME_CHECKED(UKismetTextLibrary, Conv_RotatorToText));
+			CallFunctionInputValue = CallFunction->FindPinChecked(TEXT("InRot"));
+		}
+		else if (PinSubCategoryObject == TBaseStructure<FTransform>::Get())
+		{
+			CallFunction = CreateFunctionCall(GET_MEMBER_NAME_CHECKED(UKismetTextLibrary, Conv_TransformToText));
+			CallFunctionInputValue = CallFunction->FindPinChecked(TEXT("InTrans"));
+		}
+		else if (PinSubCategoryObject == TBaseStructure<FLinearColor>::Get())
+		{
+			CallFunction = CreateFunctionCall(GET_MEMBER_NAME_CHECKED(UKismetTextLibrary, Conv_ColorToText));
+			CallFunctionInputValue = CallFunction->FindPinChecked(TEXT("InColor"));
+		}
+		else if (PinSubCategoryObject == TBaseStructure<FDateTime>::Get())
+		{
+			CallFunction = CreateFunctionCall(GET_MEMBER_NAME_CHECKED(UKismetTextLibrary, AsDateTime_DateTime));
+			CallFunctionInputValue = CallFunction->FindPinChecked(TEXT("In"));
+		}
+		else
+		{
+			check(!UE::Blueprint::Private::GetToTextSupportedScriptStruct().Contains(Cast<UScriptStruct>(PinSubCategoryObject)));
+		}
+	}
+
+	if (CallFunction == nullptr || CallFunctionInputValue == nullptr)
+	{
+		CompilerContext.MessageLog.Error(*LOCTEXT("Error_BadPinCategory", "The input pin could not generate a valid ToText.").ToString());
+		return;
+	}
+
+	// Move connection of ToText's to the generated function.
+	CompilerContext.MovePinLinksToIntermediate(*FindPinChecked(UEdGraphSchema_K2::PN_ReturnValue), *CallFunction->GetReturnValuePin());
+	CompilerContext.MovePinLinksToIntermediate(*FindPinChecked(TEXT("Value")), *CallFunctionInputValue);
+
+	BreakAllNodeLinks();
+}
+
+bool UK2Node_GenericToText::IsConnectionDisallowed(const UEdGraphPin* MyPin, const UEdGraphPin* OtherPin, FString& OutReason) const
+{
+	UEdGraphPin* InputPin = FindPinChecked(UE::Blueprint::Private::PC_Value.Resolve(), EEdGraphPinDirection::EGPD_Input);
+	if (InputPin == MyPin && MyPin->Direction == EGPD_Input)
+	{
+		const FName OtherPinCategory = OtherPin->PinType.PinCategory;	
+		bool bIsValidType = UE::Blueprint::Private::IsGenericNumericProperty(OtherPinCategory)
+			|| OtherPinCategory == UEdGraphSchema_K2::PC_String
+			|| OtherPinCategory == UEdGraphSchema_K2::PC_Name;
+
+		if (!bIsValidType)
+		{
+			const UObject* PinSubCategoryObject = OtherPin->PinType.PinSubCategoryObject.Get();
+			const bool bIsScriptStruct = OtherPinCategory == UEdGraphSchema_K2::PC_Struct
+				&& UE::Blueprint::Private::GetToTextSupportedScriptStruct().Contains(Cast<UScriptStruct>(PinSubCategoryObject));
+
+			bIsValidType = bIsScriptStruct;
+		}
+
+		if (!bIsValidType)
+		{
+			OutReason = LOCTEXT("Error_InvalidArgumentType", "To Text arguments may only be Byte, Integer, Int64, Float, Double, Text, String, Name, Boolean, Enum, and basic struct.").ToString();
+			return true;
+		}
+	}
+
+	return Super::IsConnectionDisallowed(MyPin, OtherPin, OutReason);
+}
+
+void UK2Node_GenericToText::GetMenuActions(FBlueprintActionDatabaseRegistrar& ActionRegistrar) const
+{
+	UClass* ActionKey = GetClass();
+	if (ActionRegistrar.IsOpenForRegistration(ActionKey))
+	{
+		UBlueprintNodeSpawner* NodeSpawner = UBlueprintNodeSpawner::Create(GetClass());
+		ActionRegistrar.AddBlueprintAction(ActionKey, NodeSpawner);
+	}
+}
+
+FText UK2Node_GenericToText::GetMenuCategory() const
+{
+	return FEditorCategoryUtils::GetCommonCategory(FCommonEditorCategory::Text);
+}
+
+#undef LOCTEXT_NAMESPACE
diff --git a/Engine/Source/Runtime/Engine/Classes/Kismet/KismetTextLibrary.h b/Engine/Source/Runtime/Engine/Classes/Kismet/KismetTextLibrary.h
--- a/Engine/Source/Runtime/Engine/Classes/Kismet/KismetTextLibrary.h
+++ b/Engine/Source/Runtime/Engine/Classes/Kismet/KismetTextLibrary.h
@@ -340,4 +340,9 @@ class UKismetTextLibrary : public UBlueprintFunctionLibrary
 	UFUNCTION(BlueprintCallable, CustomThunk, Category="Utilities|Text", meta=(DefaultToSelf="TextOwner"))
 	static ENGINE_API bool EditTextSourceString(UObject* TextOwner, UPARAM(ref) FText& Text, const FString& SourceString);
 	DECLARE_FUNCTION(execEditTextSourceString);
+
+	/** Converts a generic value to localized formatted text using the user culture. */
+	UFUNCTION(BlueprintPure, CustomThunk, DisplayName = "To Text", Category = "Utilities|Text", meta = (CustomStructureParam = "Value", BlueprintInternalUseOnly = "true"))
+	static ENGINE_API FText Conv_NumericPropertyToText(UPARAM(ref) const int32& Value);
+	DECLARE_FUNCTION(execConv_NumericPropertyToText);
 };
diff --git a/Engine/Source/Runtime/Engine/Private/KismetTextLibrary.cpp b/Engine/Source/Runtime/Engine/Private/KismetTextLibrary.cpp
--- a/Engine/Source/Runtime/Engine/Private/KismetTextLibrary.cpp
+++ b/Engine/Source/Runtime/Engine/Private/KismetTextLibrary.cpp
@@ -8,6 +8,7 @@
 #include "Internationalization/TextKey.h"
 #include "Internationalization/TextPackageNamespaceUtil.h"
 #include "Misc/RuntimeErrors.h"
+#include "UObject/EnumProperty.h"
 
 #define LOCTEXT_NAMESPACE "Kismet"
 
@@ -426,4 +427,62 @@ DEFINE_FUNCTION(UKismetTextLibrary::execEditTextSourceString)
 	P_NATIVE_END;
 }
 
+FText UKismetTextLibrary::Conv_NumericPropertyToText(const int32& Value)
+{
+	// We should never hit this! Stubbed to avoid NoExport on the class.
+	check(0);
+	return FText::GetEmpty();
+}
+
+DEFINE_FUNCTION(UKismetTextLibrary::execConv_NumericPropertyToText)
+{
+	Stack.StepCompiledIn<FProperty>(nullptr);
+	const FProperty* SourceProperty = Stack.MostRecentProperty;
+	void* SourceValuePtr = Stack.MostRecentPropertyAddress;
+
+	P_FINISH;
+
+	P_NATIVE_BEGIN;
+	{
+		*(FText*)RESULT_PARAM = FText::GetEmpty();
+
+		if (SourceProperty == nullptr || SourceValuePtr == nullptr)
+		{
+			LogRuntimeWarning(LOCTEXT("GenericToText.Warning.NullProperty", "The property is invalid!"));
+			return;
+		}
+
+
+		if (const FNumericProperty* NumericProperty = CastField<FNumericProperty>(SourceProperty))
+		{
+			if (NumericProperty->IsFloatingPoint())
+			{
+				double Value = NumericProperty->GetFloatingPointPropertyValue(SourceValuePtr);
+				*(FText*)RESULT_PARAM = FText::AsNumber(Value);
+			}
+			else if (UEnum* Enum = NumericProperty->GetIntPropertyEnum())
+			{
+				int64 Value = NumericProperty->GetSignedIntPropertyValue(SourceValuePtr);
+				*(FText*)RESULT_PARAM = Enum->GetDisplayNameTextByValue(Value);
+			}
+			else if (NumericProperty->IsInteger())
+			{
+				// Value from BP are always signed.
+				int64 Value = NumericProperty->GetSignedIntPropertyValue(SourceValuePtr);
+				*(FText*)RESULT_PARAM = FText::AsNumber(Value);
+			}
+		}
+		else if (const FEnumProperty* EnumProperty = CastField<FEnumProperty>(SourceProperty))
+		{
+			const int64 Value = EnumProperty->GetUnderlyingProperty()->GetSignedIntPropertyValue(SourceValuePtr);
+			*(FText*)RESULT_PARAM = EnumProperty->GetEnum()->GetDisplayNameTextByValue(Value);
+		}
+		else
+		{
+			LogRuntimeWarning(LOCTEXT("GenericToText.Warning.NotSupported", "The property not supported"));
+		}
+	}
+	P_NATIVE_END;
+}
+
 #undef LOCTEXT_NAMESPACE
