Remove expensive validation logic from 2632923
#rnx

diff --git a/Engine/Source/Editor/BlueprintGraph/Private/BlueprintActionDatabase.cpp b/Engine/Source/Editor/BlueprintGraph/Private/BlueprintActionDatabase.cpp
--- a/Engine/Source/Editor/BlueprintGraph/Private/BlueprintActionDatabase.cpp
+++ b/Engine/Source/Editor/BlueprintGraph/Private/BlueprintActionDatabase.cpp
@@ -1211,12 +1211,10 @@ FBlueprintActionDatabase::~FBlueprintActionDatabase()
 //------------------------------------------------------------------------------
 void FBlueprintActionDatabase::AddReferencedObjects(FReferenceCollector& Collector)
 {
-	TSet<UBlueprintNodeSpawner*> AllActions;
-	for (TPair<FObjectKey, FActionList>& ActionListIt : ActionRegistry)
+	for (TPair<FObjectKey, TArray<TObjectPtr<UBlueprintNodeSpawner>>>& ActionListIt : ActionRegistry)
 	{
-		FActionList& ActionList = ActionListIt.Value;
-		AllActions.Reserve(AllActions.Num() + ActionList.Num());
-		for (auto& Action : ActionList)
+		TArray<TObjectPtr<UBlueprintNodeSpawner>>& ActionList = ActionListIt.Value;
+		for (TObjectPtr<UBlueprintNodeSpawner>& Action : ActionList)
 		{
 			// We have some reports of invalid action ptrs during GC - try to catch that case here without crashing the editor while reference gathering.
 			if (!Action || (GIsGarbageCollecting && !Action->IsValidLowLevel()))
@@ -1226,38 +1224,24 @@ void FBlueprintActionDatabase::AddReferencedObjects(FReferenceCollector& Collect
 				continue;
 			}
 
-			AllActions.Add(Action);
 			Collector.AddReferencedObject(Action);
 		}
 	}
 
-	// shouldn't have to do this, as the elements listed here should also be 
-	// accounted for in the regular ActionRegistry, but just in case we fail to 
-	// remove an element from here when we should.... this'll make sure these 
-	// elements stick around (so we don't crash in ClearUnloadedAssetActions)
-	if (UnloadedActionRegistry.Num() > 0)
+	for (TPair<FSoftObjectPath, TArray<TObjectPtr<UBlueprintNodeSpawner>>>& UnloadedActionListIt : UnloadedActionRegistry)
 	{
-		TSet<UBlueprintNodeSpawner*> UnloadedActions;
-		for (TPair<FSoftObjectPath, FActionList>& UnloadedActionListIt : UnloadedActionRegistry)
+		TArray<TObjectPtr<UBlueprintNodeSpawner>>& ActionList = UnloadedActionListIt.Value;
+		for (TObjectPtr<UBlueprintNodeSpawner>& Action : ActionList)
 		{
-			FActionList& ActionList = UnloadedActionListIt.Value;
-			UnloadedActions.Reserve(UnloadedActions.Num() + ActionList.Num());
-			for (auto& Action : ActionList)
+			// Similar to above; however, we don't have any reports of failure here during GC. Nonetheless, we'll try and catch an invalid ptr value just in case.
+			if (!Action || (GIsGarbageCollecting && !Action->IsValidLowLevel()))
 			{
-				// Similar to above; however, we don't have any reports of failure here during GC. Nonetheless, we'll try and catch an invalid ptr value just in case.
-				if (!Action || (GIsGarbageCollecting && !Action->IsValidLowLevel()))
-				{
-					ensureMsgf(false, TEXT("Invalid action (0x%016llx) registered for unloaded object path: %s"), (int64)(PTRINT)Action.Get(), *UnloadedActionListIt.Key.ToString());
-					continue;
-				}
-
-				UnloadedActions.Add(Action);
-				Collector.AddReferencedObject(Action);
+				ensureMsgf(false, TEXT("Invalid action (0x%016llx) registered for unloaded object path: %s"), (int64)(PTRINT)Action.Get(), *UnloadedActionListIt.Key.ToString());
+				continue;
 			}
-		}
 
-		auto OrphanedUnloadedActions = UnloadedActions.Difference(AllActions.Intersect(UnloadedActions));
-		ensureMsgf(OrphanedUnloadedActions.Num() == 0, TEXT("Found %d unloaded actions that were not also present in the Action Registry. This should be 0."), UnloadedActions.Num());
+			Collector.AddReferencedObject(Action);
+		}
 	}
 }
 
