Make sure that Type Promotion ignores struct conversion functions which are explictly marked as "Ignore Type Promotion".

#rnx

diff --git a/Engine/Source/Editor/BlueprintGraph/Private/BlueprintTypePromotion.cpp b/Engine/Source/Editor/BlueprintGraph/Private/BlueprintTypePromotion.cpp
--- a/Engine/Source/Editor/BlueprintGraph/Private/BlueprintTypePromotion.cpp
+++ b/Engine/Source/Editor/BlueprintGraph/Private/BlueprintTypePromotion.cpp
@@ -133,10 +133,27 @@ bool FTypePromotion::HasStructConversion(const UEdGraphPin* InputPin, const UEdG
 
 	const UEdGraphSchema_K2* K2Schema = GetDefault<UEdGraphSchema_K2>();
 
-	const bool bCanAutocast = K2Schema->SearchForAutocastFunction(OutputPin->PinType, InputPin->PinType).IsSet();
+	// Check if there is a "simple" auto cast function first
+	const TOptional<UEdGraphSchema_K2::FSearchForAutocastFunctionResults> AutoCastRes = K2Schema->SearchForAutocastFunction(OutputPin->PinType, InputPin->PinType);
+	const UFunction* AutoFunc =
+		AutoCastRes.IsSet() ?
+		AutoCastRes->FunctionOwner->FindFunctionByName(AutoCastRes->TargetFunction) :
+		nullptr;
+
+	// We can auto-cast the two types so long as there is a valid function and it is not marked to ignore type promo
+	const bool bCanAutocast = AutoFunc && !AutoFunc->HasMetaData(FBlueprintMetadata::MD_IgnoreTypePromotion);
+	
+	if (bCanAutocast)
+	{
+		return true;
+	}
+
+	// Otherwise we have to search for a "special" conversion node. This is
+	// mostly for enum types but can also include some other types like containers and specific assets
 	const bool bCanAutoConvert = K2Schema->FindSpecializedConversionNode(OutputPin->PinType, *InputPin, false).IsSet();
 	
-	return bCanAutocast || bCanAutoConvert;
+	// If we have a valid result, then we can convert.
+	return bCanAutoConvert;
 }
 
 FTypePromotion::ETypeComparisonResult FTypePromotion::GetHigherType(const FEdGraphPinType& A, const FEdGraphPinType& B)
