UE-217917 - GetActor Transform nodes have hidden target pins if added to actor without link

CL 32476605 incorrectly removed a !bIsConstFunc check, which caused all pure functions to lose their target pin. This change adds that back.

diff --git a/Engine/Source/Editor/BlueprintGraph/Private/K2Node_CallFunction.cpp b/Engine/Source/Editor/BlueprintGraph/Private/K2Node_CallFunction.cpp
--- a/Engine/Source/Editor/BlueprintGraph/Private/K2Node_CallFunction.cpp
+++ b/Engine/Source/Editor/BlueprintGraph/Private/K2Node_CallFunction.cpp
@@ -1138,8 +1138,8 @@ bool UK2Node_CallFunction::CreatePinsForFunctionCall(const UFunction* Function)
 	const bool bIsProtectedFunc = Function->GetBoolMetaData(FBlueprintMetadata::MD_Protected);
 	const bool bIsStaticFunc = Function->HasAllFunctionFlags(FUNC_Static);
 
-	UEdGraph const* const Graph = GetGraph();
-	UBlueprint* BP = FBlueprintEditorUtils::FindBlueprintForGraph(Graph);
+	const UEdGraph* const Graph = GetGraph();
+	const UBlueprint* BP = FBlueprintEditorUtils::FindBlueprintForGraph(Graph);
 	ensure(BP);
 	if (BP != nullptr)
 	{
@@ -1166,9 +1166,11 @@ bool UK2Node_CallFunction::CreatePinsForFunctionCall(const UFunction* Function)
 			}
 			else
 			{
-				// Pure functions are generally compact looking, so we take advantage of that by hiding the self pin if we can.
+				const bool bLocalIsConstFunc = (Function->HasAnyFunctionFlags(FUNC_Const) != false);
+
+				// Non-const pure functions are generally compact looking, so we take advantage of that by hiding the self pin if we can.
 				// In this case, if the function belongs to our current class, then "self" is implied.
-				SelfPin->bHidden = (bIsFunctionCompatibleWithSelf && bIsPureFunc);
+				SelfPin->bHidden = (bIsFunctionCompatibleWithSelf && bIsPureFunc && !bLocalIsConstFunc);
 			}
 		}
 	}
