[UnrealEd] Moved several UFunction editor utilities from BlueprintGraph UK2Node_CallFunction to ObjectTools

#rnx

diff --git a/Engine/Plugins/Experimental/Web/WebAPI/Source/WebAPIBlueprintGraph/Private/K2Node_WebAPIOperation.cpp b/Engine/Plugins/Experimental/Web/WebAPI/Source/WebAPIBlueprintGraph/Private/K2Node_WebAPIOperation.cpp
--- a/Engine/Plugins/Experimental/Web/WebAPI/Source/WebAPIBlueprintGraph/Private/K2Node_WebAPIOperation.cpp
+++ b/Engine/Plugins/Experimental/Web/WebAPI/Source/WebAPIBlueprintGraph/Private/K2Node_WebAPIOperation.cpp
@@ -23,6 +23,7 @@
 #include "Kismet/KismetSystemLibrary.h"
 #include "Kismet2/BlueprintEditorUtils.h"
 #include "Kismet2/KismetEditorUtilities.h"
+#include "ObjectTools.h"
 
 #if WITH_EDITOR
 #include "Misc/DataValidation.h"
@@ -288,7 +289,7 @@ void UK2Node_WebAPIOperation::AllocateDefaultPins()
 
 FText UK2Node_WebAPIOperation::GetNodeTitle(ENodeTitleType::Type TitleType) const
 {
-	const FText FunctionName = UK2Node_CallFunction::GetUserFacingFunctionName(GetFactoryFunction());
+	const FText FunctionName = ObjectTools::GetUserFacingFunctionName(GetFactoryFunction());
 	FText NamespaceName;
 	FText ServiceName;
 	
@@ -339,7 +340,7 @@ FText UK2Node_WebAPIOperation::GetTooltipText() const
 	}
 	else if (CachedTooltip.IsOutOfDate(this))
 	{
-		FText BaseTooltip = FText::FromString(UK2Node_CallFunction::GetDefaultTooltipForFunction(Function));
+		FText BaseTooltip = FText::FromString(ObjectTools::GetDefaultTooltipForFunction(Function));
 
 		FFormatNamedArguments Args;
 		Args.Add(TEXT("DefaultTooltip"), BaseTooltip);
diff --git a/Engine/Source/Editor/AnimGraph/Private/AnimGraphNode_CallFunction.cpp b/Engine/Source/Editor/AnimGraph/Private/AnimGraphNode_CallFunction.cpp
--- a/Engine/Source/Editor/AnimGraph/Private/AnimGraphNode_CallFunction.cpp
+++ b/Engine/Source/Editor/AnimGraph/Private/AnimGraphNode_CallFunction.cpp
@@ -15,6 +15,7 @@
 #include "IAnimBlueprintCompilationContext.h"
 #include "Settings/EditorStyleSettings.h"
 #include "AnimBlueprintExtension_CallFunction.h"
+#include "ObjectTools.h"
 
 #define LOCTEXT_NAMESPACE "AnimGraphNode_CallFunction"
 
@@ -36,7 +37,7 @@ FText UAnimGraphNode_CallFunction::GetNodeTitle(ENodeTitleType::Type TitleType)
 	UFunction* Function = CallFunctionPrototype ? CallFunctionPrototype->GetTargetFunction() : nullptr;
 	if (Function)
 	{
-		FunctionName = UK2Node_CallFunction::GetUserFacingFunctionName(Function);
+		FunctionName = ObjectTools::GetUserFacingFunctionName(Function);
 	}
 	else if(CallFunctionPrototype)
 	{
@@ -282,9 +283,9 @@ void UAnimGraphNode_CallFunction::GetMenuActions(FBlueprintActionDatabaseRegistr
 				UBlueprintNodeSpawner* Spawner = UBlueprintNodeSpawner::Create(UAnimGraphNode_CallFunction::StaticClass(), nullptr, UBlueprintNodeSpawner::FCustomizeNodeDelegate::CreateLambda(CustomizeNode));
 				FBlueprintActionUiSpec& MenuSignature = Spawner->DefaultMenuSignature;
 
-				MenuSignature.MenuName = FText::Format(LOCTEXT("MenuNameFormat", "{0} (From Anim Graph)"), UK2Node_CallFunction::GetUserFacingFunctionName(InFunction));
+				MenuSignature.MenuName = FText::Format(LOCTEXT("MenuNameFormat", "{0} (From Anim Graph)"), ObjectTools::GetUserFacingFunctionName(InFunction));
 				MenuSignature.Category = UK2Node_CallFunction::GetDefaultCategoryForFunction(InFunction, LOCTEXT("BaseCategory", "Call Function From Anim Graph"));
-				MenuSignature.Tooltip = FText::FromString(UK2Node_CallFunction::GetDefaultTooltipForFunction(InFunction));
+				MenuSignature.Tooltip = FText::FromString(ObjectTools::GetDefaultTooltipForFunction(InFunction));
 				MenuSignature.Keywords = UK2Node_CallFunction::GetKeywordsForFunction(InFunction);
 
 				// add at least one character, so that PrimeDefaultUiSpec() doesn't attempt to query the template node
diff --git a/Engine/Source/Editor/BlueprintGraph/Classes/K2Node_CallFunction.h b/Engine/Source/Editor/BlueprintGraph/Classes/K2Node_CallFunction.h
--- a/Engine/Source/Editor/BlueprintGraph/Classes/K2Node_CallFunction.h
+++ b/Engine/Source/Editor/BlueprintGraph/Classes/K2Node_CallFunction.h
@@ -180,13 +180,15 @@ public:
 
 	virtual void PostParameterPinCreated(UEdGraphPin *Pin) {}
 
-	/** Gets the user-facing name for the function */
+	UE_DEPRECATED(5.5, "Moved to ObjectTools::GetUserFacingFunctionName.")
 	static FText GetUserFacingFunctionName(const UFunction* Function);
 
 	/** Set up a pins tooltip from a function's tooltip */
 	static void GeneratePinTooltipFromFunction(UEdGraphPin& Pin, const UFunction* Function);
-	/** Gets the non-specific tooltip for the function */
+
+	UE_DEPRECATED(5.5, "Moved to ObjectTools::GetDefaultTooltipForFunction.")
 	static FString GetDefaultTooltipForFunction(const UFunction* Function);
+
 	/** Get default category for this function in action menu */
 	static FText GetDefaultCategoryForFunction(const UFunction* Function, const FText& BaseCategory);
 	/** Get keywords for this function in the action menu */
diff --git a/Engine/Source/Editor/BlueprintGraph/Private/BlueprintFunctionNodeSpawner.cpp b/Engine/Source/Editor/BlueprintGraph/Private/BlueprintFunctionNodeSpawner.cpp
--- a/Engine/Source/Editor/BlueprintGraph/Private/BlueprintFunctionNodeSpawner.cpp
+++ b/Engine/Source/Editor/BlueprintGraph/Private/BlueprintFunctionNodeSpawner.cpp
@@ -34,6 +34,7 @@
 #include "Kismet2/BlueprintEditorUtils.h"
 #include "Misc/AssertionMacros.h"
 #include "ObjectEditorUtils.h"
+#include "ObjectTools.h"
 #include "SNodePanel.h"
 #include "Templates/Casts.h"
 #include "Textures/SlateIcon.h"
@@ -302,9 +303,9 @@ UBlueprintFunctionNodeSpawner* UBlueprintFunctionNodeSpawner::Create(TSubclassOf
 	}
 	else
 	{
-		MenuSignature.MenuName = UK2Node_CallFunction::GetUserFacingFunctionName(Function);
+		MenuSignature.MenuName = ObjectTools::GetUserFacingFunctionName(Function);
 		MenuSignature.Category = UK2Node_CallFunction::GetDefaultCategoryForFunction(Function, FText::GetEmpty());
-		MenuSignature.Tooltip = FText::FromString(UK2Node_CallFunction::GetDefaultTooltipForFunction(Function));
+		MenuSignature.Tooltip = FText::FromString(ObjectTools::GetDefaultTooltipForFunction(Function));
 		// add at least one character, so that PrimeDefaultUiSpec() doesn't attempt to query the template node
 		MenuSignature.Keywords = UK2Node_CallFunction::GetKeywordsForFunction(Function);
 	}
diff --git a/Engine/Source/Editor/BlueprintGraph/Private/EdGraphSchema_K2.cpp b/Engine/Source/Editor/BlueprintGraph/Private/EdGraphSchema_K2.cpp
--- a/Engine/Source/Editor/BlueprintGraph/Private/EdGraphSchema_K2.cpp
+++ b/Engine/Source/Editor/BlueprintGraph/Private/EdGraphSchema_K2.cpp
@@ -101,6 +101,7 @@
 
 #include "Editor/EditorPerProjectUserSettings.h"
 #include "BlueprintPaletteFavorites.h"
+#include "ObjectTools.h"
 
 //////////////////////////////////////////////////////////////////////////
 
@@ -967,7 +968,7 @@ bool UEdGraphSchema_K2::FunctionCanBeUsedInDelegate(const UFunction* InFunction)
 
 FText UEdGraphSchema_K2::GetFriendlySignatureName(const UFunction* Function)
 {
-	return UK2Node_CallFunction::GetUserFacingFunctionName( Function );
+	return ObjectTools::GetUserFacingFunctionName( Function );
 }
 
 void UEdGraphSchema_K2::GetAutoEmitTermParameters(const UFunction* Function, TArray<FString>& AutoEmitParameterNames)
@@ -5622,7 +5623,7 @@ void UEdGraphSchema_K2::GetGraphDisplayInformation(const UEdGraph& Graph, /*out*
 			if (Function)
 			{
 				DisplayInfo.PlainName = FText::FromString(Function->GetName());
-				DisplayInfo.Tooltip = FText::FromString(UK2Node_CallFunction::GetDefaultTooltipForFunction(Function)); // grab its tooltip
+				DisplayInfo.Tooltip = FText::FromString(ObjectTools::GetDefaultTooltipForFunction(Function)); // grab its tooltip
 			}
 			else
 			{
diff --git a/Engine/Source/Editor/BlueprintGraph/Private/K2Node_BaseAsyncTask.cpp b/Engine/Source/Editor/BlueprintGraph/Private/K2Node_BaseAsyncTask.cpp
--- a/Engine/Source/Editor/BlueprintGraph/Private/K2Node_BaseAsyncTask.cpp
+++ b/Engine/Source/Editor/BlueprintGraph/Private/K2Node_BaseAsyncTask.cpp
@@ -34,6 +34,7 @@
 #include "Misc/AssertionMacros.h"
 #include "Misc/ConfigCacheIni.h"
 #include "Misc/Parse.h"
+#include "ObjectTools.h"
 #include "Templates/Casts.h"
 #include "Trace/Detail/Channel.h"
 #include "UObject/Class.h"
@@ -62,7 +63,7 @@ UK2Node_BaseAsyncTask::UK2Node_BaseAsyncTask(const FObjectInitializer& ObjectIni
 FText UK2Node_BaseAsyncTask::GetTooltipText() const
 {
 	FFormatNamedArguments Args;
-	Args.Add(TEXT("FunctionTooltip"), FText::FromString(UK2Node_CallFunction::GetDefaultTooltipForFunction(GetFactoryFunction())));
+	Args.Add(TEXT("FunctionTooltip"), FText::FromString(ObjectTools::GetDefaultTooltipForFunction(GetFactoryFunction())));
 	Args.Add(TEXT("LatentString"), NSLOCTEXT("K2Node", "LatentFunction", "Latent. This node will complete at a later time. Latent nodes can only be placed in event graphs."));
 	
 	return FText::Format(LOCTEXT("AsyncTaskTooltip", "{FunctionTooltip}\n\n{LatentString}"), Args);
@@ -74,7 +75,7 @@ FText UK2Node_BaseAsyncTask::GetNodeTitle(ENodeTitleType::Type TitleType) const
 	{
 		return FText(LOCTEXT("UK2Node_BaseAsyncTaskGetNodeTitle", "Async Task: Missing Function"));
 	}
-	const FText FunctionToolTipText = UK2Node_CallFunction::GetUserFacingFunctionName(GetFactoryFunction());
+	const FText FunctionToolTipText = ObjectTools::GetUserFacingFunctionName(GetFactoryFunction());
 	return FunctionToolTipText;
 }
 
diff --git a/Engine/Source/Editor/BlueprintGraph/Private/K2Node_CallFunction.cpp b/Engine/Source/Editor/BlueprintGraph/Private/K2Node_CallFunction.cpp
--- a/Engine/Source/Editor/BlueprintGraph/Private/K2Node_CallFunction.cpp
+++ b/Engine/Source/Editor/BlueprintGraph/Private/K2Node_CallFunction.cpp
@@ -48,6 +48,7 @@
 #include "HAL/FileManager.h"
 #include "Widgets/Notifications/SNotificationList.h"
 #include "BlueprintNodeStatics.h"
+#include "ObjectTools.h"
 #include "Settings/BlueprintEditorProjectSettings.h"
 #include "ToolMenu.h"
 
@@ -565,7 +566,7 @@ FEdGraphNodeDeprecationResponse UK2Node_CallFunction::GetDeprecationResponse(EEd
 			if (ensureMsgf(Function != nullptr, TEXT("This node should not be able to report having a deprecated reference if the target function cannot be resolved.")))
 			{
 				FString DetailedMessage = Function->GetMetaData(FBlueprintMetadata::MD_DeprecationMessage);
-				Response.MessageText = FBlueprintEditorUtils::GetDeprecatedMemberUsageNodeWarning(GetUserFacingFunctionName(Function), FText::FromString(DetailedMessage));
+				Response.MessageText = FBlueprintEditorUtils::GetDeprecatedMemberUsageNodeWarning(ObjectTools::GetUserFacingFunctionName(Function), FText::FromString(DetailedMessage));
 			}
 		}
 	}
@@ -609,7 +610,7 @@ FText UK2Node_CallFunction::GetNodeTitle(ENodeTitleType::Type TitleType) const
 	if (UFunction* Function = GetTargetFunction())
 	{
 		RPCString = UK2Node_Event::GetLocalizedNetString(Function->FunctionFlags, true);
-		FunctionName = GetUserFacingFunctionName(Function);
+		FunctionName = ObjectTools::GetUserFacingFunctionName(Function);
 		ContextString = GetFunctionContextString();
 	}
 	else
@@ -1627,7 +1628,7 @@ FText UK2Node_CallFunction::GetTooltipText() const
 	}
 	else if (CachedTooltip.IsOutOfDate(this))
 	{
-		FText BaseTooltip = FText::FromString(GetDefaultTooltipForFunction(Function));
+		FText BaseTooltip = FText::FromString(ObjectTools::GetDefaultTooltipForFunction(Function));
 
 		FFormatNamedArguments Args;
 		Args.Add(TEXT("DefaultTooltip"), BaseTooltip);
@@ -1797,70 +1798,12 @@ void UK2Node_CallFunction::GeneratePinTooltipFromFunction(UEdGraphPin& Pin, cons
 
 FText UK2Node_CallFunction::GetUserFacingFunctionName(const UFunction* Function)
 {
-	FText ReturnDisplayName;
-
-	if (Function != NULL)
-	{
-		if (GEditor && GetDefault<UEditorStyleSettings>()->bShowFriendlyNames)
-		{
-			ReturnDisplayName = Function->GetDisplayNameText();
-		}
-		else
-		{
-			static const FString Namespace = TEXT("UObjectDisplayNames");
-			const FString Key = Function->GetFullGroupName(false);
-
-			ReturnDisplayName = Function->GetMetaDataText(TEXT("DisplayName"), Namespace, Key);
-		}
-	}
-	return ReturnDisplayName;
+	return ObjectTools::GetUserFacingFunctionName(Function);
 }
 
 FString UK2Node_CallFunction::GetDefaultTooltipForFunction(const UFunction* Function)
 {
-	FString Tooltip;
-
-	if (Function != NULL)
-	{
-		Tooltip = Function->GetToolTipText().ToString();
-	}
-
-	if (!Tooltip.IsEmpty())
-	{
-		// Strip off the doxygen nastiness
-		static const FString DoxygenParam(TEXT("@param"));
-		static const FString DoxygenReturn(TEXT("@return"));
-		static const FString DoxygenSee(TEXT("@see"));
-		static const FString TooltipSee(TEXT("See:"));
-		static const FString DoxygenNote(TEXT("@note"));
-		static const FString TooltipNote(TEXT("Note:"));
-
-		Tooltip.Split(DoxygenParam, &Tooltip, nullptr, ESearchCase::IgnoreCase, ESearchDir::FromStart);
-		Tooltip.Split(DoxygenReturn, &Tooltip, nullptr, ESearchCase::IgnoreCase, ESearchDir::FromStart);
-
-		Tooltip.ReplaceInline(*DoxygenSee, *TooltipSee);
-		Tooltip.ReplaceInline(*DoxygenNote, *TooltipNote);
-
-		Tooltip.TrimStartAndEndInline();
-
-		UClass* CurrentSelfClass = (Function != NULL) ? Function->GetOwnerClass() : NULL;
-		UClass const* TrueSelfClass = CurrentSelfClass;
-		if (CurrentSelfClass && CurrentSelfClass->ClassGeneratedBy)
-		{
-			TrueSelfClass = CurrentSelfClass->GetAuthoritativeClass();
-		}
-
-		FText TargetDisplayText = (TrueSelfClass != NULL) ? TrueSelfClass->GetDisplayNameText() : LOCTEXT("None", "None");
-
-		FFormatNamedArguments Args;
-		Args.Add(TEXT("TargetName"), TargetDisplayText);
-		Args.Add(TEXT("Tooltip"), FText::FromString(Tooltip));
-		return FText::Format(LOCTEXT("CallFunction_Tooltip", "{Tooltip}\n\nTarget is {TargetName}"), Args).ToString();
-	}
-	else
-	{
-		return GetUserFacingFunctionName(Function).ToString();
-	}
+	return ObjectTools::GetDefaultTooltipForFunction(Function);
 }
 
 FText UK2Node_CallFunction::GetDefaultCategoryForFunction(const UFunction* Function, const FText& BaseCategory)
diff --git a/Engine/Source/Editor/BlueprintGraph/Private/K2Node_CallParentFunction.cpp b/Engine/Source/Editor/BlueprintGraph/Private/K2Node_CallParentFunction.cpp
--- a/Engine/Source/Editor/BlueprintGraph/Private/K2Node_CallParentFunction.cpp
+++ b/Engine/Source/Editor/BlueprintGraph/Private/K2Node_CallParentFunction.cpp
@@ -14,6 +14,7 @@
 #include "Internationalization/Internationalization.h"
 #include "K2Node.h"
 #include "Misc/Guid.h"
+#include "ObjectTools.h"
 #include "Settings/EditorStyleSettings.h"
 #include "Templates/Casts.h"
 #include "UObject/Class.h"
@@ -34,7 +35,7 @@ FText UK2Node_CallParentFunction::GetNodeTitle(ENodeTitleType::Type TitleType) c
 
 	if (Function)
 	{
-		FunctionName = GetUserFacingFunctionName( Function );
+		FunctionName = ObjectTools::GetUserFacingFunctionName( Function );
 	}
 	else if ( GEditor && GetDefault<UEditorStyleSettings>()->bShowFriendlyNames )
 	{
diff --git a/Engine/Source/Editor/BlueprintGraph/Private/K2Node_Event.cpp b/Engine/Source/Editor/BlueprintGraph/Private/K2Node_Event.cpp
--- a/Engine/Source/Editor/BlueprintGraph/Private/K2Node_Event.cpp
+++ b/Engine/Source/Editor/BlueprintGraph/Private/K2Node_Event.cpp
@@ -29,6 +29,7 @@
 #include "Logging/LogCategory.h"
 #include "Logging/LogMacros.h"
 #include "Misc/AssertionMacros.h"
+#include "ObjectTools.h"
 #include "Serialization/Archive.h"
 #include "Styling/AppStyle.h"
 #include "Templates/Casts.h"
@@ -166,7 +167,7 @@ FText UK2Node_Event::GetTooltipText() const
 	UFunction* Function = EventReference.ResolveMember<UFunction>(GetBlueprintClassFromNode());
 	if (CachedTooltip.IsOutOfDate(this) && (Function != nullptr))
 	{
-		CachedTooltip.SetCachedText(FText::FromString(UK2Node_CallFunction::GetDefaultTooltipForFunction(Function)), this);
+		CachedTooltip.SetCachedText(FText::FromString(ObjectTools::GetDefaultTooltipForFunction(Function)), this);
 
 		if (bOverrideFunction || (CustomFunctionName == NAME_None))
 		{
diff --git a/Engine/Source/Editor/BlueprintGraph/Private/K2Node_FunctionEntry.cpp b/Engine/Source/Editor/BlueprintGraph/Private/K2Node_FunctionEntry.cpp
--- a/Engine/Source/Editor/BlueprintGraph/Private/K2Node_FunctionEntry.cpp
+++ b/Engine/Source/Editor/BlueprintGraph/Private/K2Node_FunctionEntry.cpp
@@ -37,6 +37,7 @@
 #include "Logging/LogCategory.h"
 #include "Logging/LogMacros.h"
 #include "Misc/AssertionMacros.h"
+#include "ObjectTools.h"
 #include "Serialization/Archive.h"
 #include "Templates/Casts.h"
 #include "Templates/SubclassOf.h"
@@ -684,7 +685,7 @@ FText UK2Node_FunctionEntry::GetTooltipText() const
 {
 	if (UFunction* const Function = FindSignatureFunction())
 	{
-		return FText::FromString(UK2Node_CallFunction::GetDefaultTooltipForFunction(Function));
+		return FText::FromString(ObjectTools::GetDefaultTooltipForFunction(Function));
 	}
 	return Super::GetTooltipText();
 }
diff --git a/Engine/Source/Editor/BlueprintGraph/Private/K2Node_Message.cpp b/Engine/Source/Editor/BlueprintGraph/Private/K2Node_Message.cpp
--- a/Engine/Source/Editor/BlueprintGraph/Private/K2Node_Message.cpp
+++ b/Engine/Source/Editor/BlueprintGraph/Private/K2Node_Message.cpp
@@ -27,6 +27,7 @@
 #include "Logging/LogCategory.h"
 #include "Logging/LogMacros.h"
 #include "Misc/AssertionMacros.h"
+#include "ObjectTools.h"
 #include "Templates/Casts.h"
 #include "Templates/SubclassOf.h"
 #include "Trace/Detail/Channel.h"
@@ -50,7 +51,7 @@ FText UK2Node_Message::GetNodeTitle(ENodeTitleType::Type TitleType) const
 	{
 		if (!CachedNodeTitles.IsTitleCached(TitleType, this))
 		{
-			FText NodeNameText = UK2Node_CallFunction::GetUserFacingFunctionName(Function);
+			FText NodeNameText = ObjectTools::GetUserFacingFunctionName(Function);
 			if (TitleType == ENodeTitleType::MenuTitle)
 			{
 				// FText::Format() is slow, so we cache this to save on performance
diff --git a/Engine/Source/Editor/DetailCustomizations/Private/ObjectDetails.cpp b/Engine/Source/Editor/DetailCustomizations/Private/ObjectDetails.cpp
--- a/Engine/Source/Editor/DetailCustomizations/Private/ObjectDetails.cpp
+++ b/Engine/Source/Editor/DetailCustomizations/Private/ObjectDetails.cpp
@@ -17,13 +17,13 @@
 #include "HAL/PlatformCrt.h"
 #include "Internationalization/Internationalization.h"
 #include "Internationalization/Text.h"
-#include "K2Node_CallFunction.h"
 #include "Kismet2/BlueprintEditorUtils.h"
 #include "Layout/Margin.h"
 #include "Math/NumericLimits.h"
 #include "Misc/Attribute.h"
 #include "Misc/CString.h"
 #include "ObjectEditorUtils.h"
+#include "ObjectTools.h"
 #include "PropertyPermissionList.h"
 #include "Reflection/FunctionUtils.h"
 #include "Settings/BlueprintEditorProjectSettings.h"
@@ -244,7 +244,7 @@ void FObjectDetails::AddCallInEditorMethods(IDetailLayoutBuilder& DetailBuilder)
 			}
 			FCategoryEntry& CategoryEntry = CategoryList.Last();
 
-			const FText ButtonCaption = UK2Node_CallFunction::GetUserFacingFunctionName(Function);
+			const FText ButtonCaption = ObjectTools::GetUserFacingFunctionName(Function);
 			FText FunctionTooltip = Function->GetToolTipText();
 			if (FunctionTooltip.IsEmpty())
 			{
diff --git a/Engine/Source/Editor/Kismet/Private/SMyBlueprint.cpp b/Engine/Source/Editor/Kismet/Private/SMyBlueprint.cpp
--- a/Engine/Source/Editor/Kismet/Private/SMyBlueprint.cpp
+++ b/Engine/Source/Editor/Kismet/Private/SMyBlueprint.cpp
@@ -73,6 +73,7 @@
 #include "Misc/CString.h"
 #include "Misc/Guid.h"
 #include "ObjectEditorUtils.h"
+#include "ObjectTools.h"
 #include "SBlueprintPalette.h"
 #include "SGraphActionMenu.h"
 #include "SKismetInspector.h"
@@ -1598,7 +1599,7 @@ void SMyBlueprint::CollectAllActions(FGraphActionListBuilderBase& OutAllActions)
 			 && Blueprint->AllowFunctionOverride(Function)
 		   )
 		{
-			FText FunctionTooltip = FText::FromString(UK2Node_CallFunction::GetDefaultTooltipForFunction(Function));
+			FText FunctionTooltip = FText::FromString(ObjectTools::GetDefaultTooltipForFunction(Function));
 			FText FunctionDesc = K2Schema->GetFriendlySignatureName(Function);
 			if ( FunctionDesc.IsEmpty() )
 			{
diff --git a/Engine/Source/Editor/UnrealEd/Private/ObjectTools.cpp b/Engine/Source/Editor/UnrealEd/Private/ObjectTools.cpp
--- a/Engine/Source/Editor/UnrealEd/Private/ObjectTools.cpp
+++ b/Engine/Source/Editor/UnrealEd/Private/ObjectTools.cpp
@@ -106,9 +106,13 @@
 #include "Logging/LogMacros.h"
 #include "UncontrolledChangelistsModule.h"
 #include "AssetCompilingManager.h"
+#include "ObjectEditorUtils.h"
+#include "Settings/EditorStyleSettings.h"
 
 DEFINE_LOG_CATEGORY_STATIC(LogObjectTools, Log, All);
 
+#define LOCTEXT_NAMESPACE "ObjectTools"
+
 static TAutoConsoleVariable<bool> CVarUseLegacyGetReferencersForDeletion(
 	TEXT("Editor.UseLegacyGetReferencersForDeletion"),
 	false,
@@ -5026,6 +5030,75 @@ namespace ObjectTools
 			}
 		}
 	}
+
+	FText GetUserFacingFunctionName(const UFunction* Function)
+	{
+		FText ReturnDisplayName;
+
+		if (Function != nullptr)
+		{
+			if (GEditor && GetDefault<UEditorStyleSettings>()->bShowFriendlyNames)
+			{
+				ReturnDisplayName = Function->GetDisplayNameText();
+			}
+			else
+			{
+				static const FString Namespace = TEXT("UObjectDisplayNames");
+				const FString Key = Function->GetFullGroupName(false);
+
+				ReturnDisplayName = Function->GetMetaDataText(TEXT("DisplayName"), Namespace, Key);
+			}
+		}
+
+		return ReturnDisplayName;
+	}
+
+	FString GetDefaultTooltipForFunction(const UFunction* Function)
+	{
+		FString Tooltip;
+
+		if (Function != nullptr)
+		{
+			Tooltip = Function->GetToolTipText().ToString();
+		}
+
+		if (!Tooltip.IsEmpty())
+		{
+			// Strip off the doxygen nastiness
+			static const FString DoxygenParam(TEXT("@param"));
+			static const FString DoxygenReturn(TEXT("@return"));
+			static const FString DoxygenSee(TEXT("@see"));
+			static const FString TooltipSee(TEXT("See:"));
+			static const FString DoxygenNote(TEXT("@note"));
+			static const FString TooltipNote(TEXT("Note:"));
+
+			Tooltip.Split(DoxygenParam, &Tooltip, nullptr, ESearchCase::IgnoreCase, ESearchDir::FromStart);
+			Tooltip.Split(DoxygenReturn, &Tooltip, nullptr, ESearchCase::IgnoreCase, ESearchDir::FromStart);
+
+			Tooltip.ReplaceInline(*DoxygenSee, *TooltipSee);
+			Tooltip.ReplaceInline(*DoxygenNote, *TooltipNote);
+
+			Tooltip.TrimStartAndEndInline();
+
+			UClass* CurrentSelfClass = (Function != nullptr) ? Function->GetOwnerClass() : nullptr;
+			UClass const* TrueSelfClass = CurrentSelfClass;
+			if (CurrentSelfClass && CurrentSelfClass->ClassGeneratedBy)
+			{
+				TrueSelfClass = CurrentSelfClass->GetAuthoritativeClass();
+			}
+
+			FText TargetDisplayText = (TrueSelfClass != nullptr) ? TrueSelfClass->GetDisplayNameText() : LOCTEXT("None", "None");
+
+			FFormatNamedArguments Args;
+			Args.Add(TEXT("TargetName"), TargetDisplayText);
+			Args.Add(TEXT("Tooltip"), FText::FromString(Tooltip));
+			return FText::Format(LOCTEXT("CallFunction_Tooltip", "{Tooltip}\n\nTarget is {TargetName}"), Args).ToString();
+		}
+		else
+		{
+			return GetUserFacingFunctionName(Function).ToString();
+		}
+	}
 }
 
 
@@ -5831,4 +5904,6 @@ namespace ThumbnailTools
 
 		return false;
 	}
-		}
+}
+
+#undef LOCTEXT_NAMESPACE
diff --git a/Engine/Source/Editor/UnrealEd/Public/ObjectTools.h b/Engine/Source/Editor/UnrealEd/Public/ObjectTools.h
--- a/Engine/Source/Editor/UnrealEd/Public/ObjectTools.h
+++ b/Engine/Source/Editor/UnrealEd/Public/ObjectTools.h
@@ -674,6 +674,22 @@ namespace ObjectTools
 	 * @param	OutInstances	per InObjects array element, receives the list of objects which have the given object in its archetype chain
 	 */
 	UNREALED_API void BatchGetArchetypeInstances(TArrayView<UObject*> InObjects, TArray<TArray<UObject*>>& OutInstances);
+
+
+	/**
+	 * Gets the user-facing name for the given UFunction.
+	 * @param Function	the function to find the name for
+	 * @return the name as localized text
+	 */
+	UNREALED_API FText GetUserFacingFunctionName(const UFunction* Function);
+
+
+	/**
+	 * Gets the non-specific tooltip for the given UFunction.
+	 * @param Function the function to find the tooltip for
+	 * @return the tooltip as a string
+	 */
+	UNREALED_API FString GetDefaultTooltipForFunction(const UFunction* Function);
 }
 
 
