Rework wildcard inference to infer types, rather than just find the first connected pin that isn't a wildcard

[RN] Blueprint macros can now infer more than one type of pin, especially useful for TMap related macros

diff --git a/Engine/Source/Editor/BlueprintGraph/Classes/K2Node_MacroInstance.h b/Engine/Source/Editor/BlueprintGraph/Classes/K2Node_MacroInstance.h
--- a/Engine/Source/Editor/BlueprintGraph/Classes/K2Node_MacroInstance.h
+++ b/Engine/Source/Editor/BlueprintGraph/Classes/K2Node_MacroInstance.h
@@ -103,12 +103,25 @@ public:
 	void SetMacroGraph(UEdGraph* Graph) { MacroGraphReference.SetGraph(Graph); }
 	UEdGraph* GetMacroGraph() const { return MacroGraphReference.GetGraph(); }
 	UBlueprint* GetSourceBlueprint() const { return MacroGraphReference.GetBlueprint(); }
+	// Public function for inferring wildcards after macro expansion. Intended
+	// for use only by the blueprint compiler. Inference is typically performed
+	// automatically when a pin connection is made.
+	BLUEPRINTGRAPH_API void InferWildcards(const TArray<UEdGraphNode*>& InNodes) const;
 
 	// Finds the associated metadata for the macro instance if there is any; this function is not particularly fast.
 	BLUEPRINTGRAPH_API static FKismetUserDeclaredFunctionMetadata* GetAssociatedGraphMetadata(const UEdGraph* AssociatedMacroGraph);
 	static void FindInContentBrowser(TWeakObjectPtr<UK2Node_MacroInstance> MacroInstance);
 
 private:
+	// Helpers and implementation of 'smart' wildcard inference - which doesn't
+	// just grab the first user type and overwrite every wildcard pin with it.
+	// Instead it intelligently propagates types from the macro inputs throughout
+	// the graph, and then back to any input pins:
+	void InferWildcards();
+	TArray<UEdGraphPin*> GetAllWildcardPins() const;
+	static bool ShouldDoSmartWildcardInference();
+	void SmartInferWildcardsImpl(const TArray<UEdGraphNode*>& InNodes) const;
+
 	/** Constructing FText strings can be costly, so we cache the node's tooltip */
 	FNodeTextCache CachedTooltip;
 };
diff --git a/Engine/Source/Editor/BlueprintGraph/Private/K2Node_GetArrayItem.cpp b/Engine/Source/Editor/BlueprintGraph/Private/K2Node_GetArrayItem.cpp
--- a/Engine/Source/Editor/BlueprintGraph/Private/K2Node_GetArrayItem.cpp
+++ b/Engine/Source/Editor/BlueprintGraph/Private/K2Node_GetArrayItem.cpp
@@ -450,7 +450,9 @@ void UK2Node_GetArrayItem::PropagatePinType(FEdGraphPinType& InType)
 	ArrayPin->PinType.bIsReference = false;
 
 	// IsSetToReturnRef() has to be called after the ArrayPin's type is set, since it uses that to determine the result
-	const bool bMakeOutputRef = IsSetToReturnRef();
+	const bool bMakeOutputRef = IsSetToReturnRef() ||
+		// Honor bReturnByRefDesired for transient (intermediate) nodes during compilation:
+		(Blueprint && Blueprint->bBeingCompiled && bReturnByRefDesired && HasAnyFlags(RF_Transient));
 	if (!bMakeOutputRef && bReturnByRefDesired && ResultPin->PinType.bIsReference && Blueprint && !Blueprint->bIsRegeneratingOnLoad)
 	{
 		FNotificationInfo Warning(LOCTEXT("ConnectionAlteredOutput", "Array Get node altered. Now returning a copy."));
diff --git a/Engine/Source/Editor/BlueprintGraph/Private/K2Node_MacroInstance.cpp b/Engine/Source/Editor/BlueprintGraph/Private/K2Node_MacroInstance.cpp
--- a/Engine/Source/Editor/BlueprintGraph/Private/K2Node_MacroInstance.cpp
+++ b/Engine/Source/Editor/BlueprintGraph/Private/K2Node_MacroInstance.cpp
@@ -7,6 +7,7 @@
 #include "Delegates/Delegate.h"
 #include "EdGraph/EdGraphSchema.h"
 #include "EdGraphSchema_K2.h"
+#include "EdGraphUtilities.h"
 #include "Editor.h"
 #include "Editor/EditorEngine.h"
 #include "EditorCategoryUtils.h"
@@ -16,6 +17,8 @@
 #include "Internationalization/Internationalization.h"
 #include "K2Node_EditablePinBase.h"
 #include "Kismet2/BlueprintEditorUtils.h"
+#include "Kismet2/WildcardNodeUtils.h"
+#include "KismetCompiler.h"
 #include "Misc/AssertionMacros.h"
 #include "Serialization/Archive.h"
 #include "Settings/EditorStyleSettings.h"
@@ -291,16 +294,20 @@ void UK2Node_MacroInstance::NotifyPinConnectionListChanged(UEdGraphPin* ChangedP
 		{
 			// get type of pin we just got linked to
 			FEdGraphPinType const LinkedPinType = ChangedPin->LinkedTo[0]->PinType;
-
-			// change all other wildcard pins to the new type
-			// note we're assuming only one wildcard type per Macro node, for now
-
-			for(int32 PinIdx=0; PinIdx<Pins.Num(); PinIdx++)
+			if(ShouldDoSmartWildcardInference())
 			{
-				UEdGraphPin* const TmpPin = Pins[PinIdx];
-				if (TmpPin)
+				// only copy the category stuff to preserve array and ref status
+				FWildcardNodeUtils::InferType(ChangedPin, LinkedPinType);
+				InferWildcards();
+			}
+			else
+			{
+				// change all other wildcard pins to the new type
+				// note we're assuming only one wildcard type per Macro node, for now
+				for(int32 PinIdx=0; PinIdx<Pins.Num(); PinIdx++)
 				{
-					if (TmpPin->PinType.PinCategory == UEdGraphSchema_K2::PC_Wildcard)
+					UEdGraphPin* const TmpPin = Pins[PinIdx];
+					if (FWildcardNodeUtils::IsWildcardPin(TmpPin))
 					{
 						// only copy the category stuff to preserve array and ref status
 						TmpPin->PinType.PinCategory = LinkedPinType.PinCategory;
@@ -321,8 +328,6 @@ void UK2Node_MacroInstance::NotifyPinConnectionListChanged(UEdGraphPin* ChangedP
 	}
 }
 
-
-
 void UK2Node_MacroInstance::NodeConnectionListChanged()
 {
 	Super::NodeConnectionListChanged();
@@ -479,6 +484,35 @@ void UK2Node_MacroInstance::PostFixupAllWildcardPins(bool bInAllWildcardPinsUnli
 	}
 }
 
+void UK2Node_MacroInstance::InferWildcards(const TArray<UEdGraphNode*>& InNodes) const
+{
+	if(ShouldDoSmartWildcardInference())
+	{
+		SmartInferWildcardsImpl(InNodes);
+		return;
+	}
+
+	if (!ResolvedWildcardType.PinCategory.IsNone())
+	{
+		for (UEdGraphNode* const ClonedNode : InNodes)
+		{
+			if (ClonedNode)
+			{
+				for (UEdGraphPin* const ClonedPin : ClonedNode->Pins)
+				{
+					if (ClonedPin && (ClonedPin->PinType.PinCategory == UEdGraphSchema_K2::PC_Wildcard))
+					{
+						// copy only type info, so array or ref status is preserved
+						ClonedPin->PinType.PinCategory = ResolvedWildcardType.PinCategory;
+						ClonedPin->PinType.PinSubCategory = ResolvedWildcardType.PinSubCategory;
+						ClonedPin->PinType.PinSubCategoryObject = ResolvedWildcardType.PinSubCategoryObject;
+					}
+				}
+			}
+		}
+	}
+}
+
 bool UK2Node_MacroInstance::HasExternalDependencies(TArray<class UStruct*>* OptionalOutput) const
 {
 	UBlueprint* OtherBlueprint = MacroGraphReference.GetBlueprint();
@@ -546,4 +580,236 @@ FBlueprintNodeSignature UK2Node_MacroInstance::GetSignature() const
 	return NodeSignature;
 }
 
+void UK2Node_MacroInstance::InferWildcards()
+{
+	// we've got a new user provided pin, expand the macro 
+	UEdGraph* MacroGraph = GetMacroGraph();
+	if (MacroGraph)
+	{
+		// perform macro expansion in a dummy graph, inferring whatever types we can from the provided wildcards:
+		FCompilerResultsLog MessageLog;
+		UBlueprint* BP = GetBlueprint();
+		UEdGraph* ClonedGraph = FEdGraphUtilities::CloneGraph(MacroGraph, GetBlueprint(), &MessageLog, true);
+		if (ClonedGraph)
+		{
+			InferWildcards(ClonedGraph->Nodes);
+
+			// Uncomment to record this graph as an intermediate product - useful for debugging
+			//ClonedGraph->Schema = UEdGraphSchema_K2::StaticClass();
+			//GetBlueprint()->IntermediateGeneratedGraphs.Add(ClonedGraph);
+			//ClonedGraph->SetFlags(RF_Transient);
+		}
+	}
+}
+
+TArray<UEdGraphPin*> UK2Node_MacroInstance::GetAllWildcardPins() const
+{
+	UEdGraph* MacroGraph = GetMacroGraph();
+	if (MacroGraph == nullptr)
+	{
+		return TArray<UEdGraphPin*>();
+	}
+
+	TArray<UEdGraphPin*> Result;
+	for(UEdGraphNode* Node : MacroGraph->Nodes)
+	{
+		if (UK2Node_Tunnel* Tunnel = ExactCast<UK2Node_Tunnel>(Node))
+		{
+			for(UEdGraphPin* TunnelPin : Tunnel->Pins)
+			{
+				if(FWildcardNodeUtils::IsWildcardPin(TunnelPin))
+				{
+					Result.Add(FindPin(TunnelPin->GetName(), UEdGraphPin::GetComplementaryDirection(TunnelPin->Direction)));
+				}
+			}
+		}
+	}
+
+	return Result;
+}
+
+bool UK2Node_MacroInstance::ShouldDoSmartWildcardInference()
+{
+	static const FBoolConfigValueHelper bUseSimpleWildcardInference(TEXT("Blueprints"), TEXT("bUseSimpleWildcardInference"), GEngineIni);
+	return !bUseSimpleWildcardInference;
+}
+
+namespace UE::Private
+{
+
+// Recursively infers a type for a network of wildcard pins - i.e. infers type for your linkedto's linkedto's linketo's....
+static void InferLinkedPinsImpl(UEdGraphPin* Pin, const FEdGraphPinType& Type, TArray<TPair<UEdGraphNode*, UEdGraphPin*>>& OutDirtyNodePins, TSet<UEdGraphPin*>& ProcessedPins)
+{
+	FWildcardNodeUtils::InferType(Pin, Type);
+	OutDirtyNodePins.AddUnique({Pin->GetOwningNode(), Pin});
+	ProcessedPins.Add(Pin);
+
+	for(UEdGraphPin* LinkedPin : Pin->LinkedTo)
+	{
+		if(!ProcessedPins.Contains(LinkedPin) && FWildcardNodeUtils::IsWildcardPin(LinkedPin))
+		{
+			InferLinkedPinsImpl(LinkedPin, Type, OutDirtyNodePins, ProcessedPins);
+		}
+	}
+}
+
+}
+
+void UK2Node_MacroInstance::SmartInferWildcardsImpl(const TArray<UEdGraphNode*>& InNodes) const
+{
+	// Gather wild card pins on the tunnel pins:
+	TArray<UEdGraphPin*> TunnelWildcards;
+	for (UEdGraphNode* Node : InNodes)
+	{
+		if (UK2Node_Tunnel* Tunnel = ExactCast<UK2Node_Tunnel>(Node))
+		{
+			for (UEdGraphPin* TunnelPin : Tunnel->Pins)
+			{
+				if (FWildcardNodeUtils::IsWildcardPin(TunnelPin))
+				{
+					// the tunnel node with input pins is the output on the macro:
+					TunnelWildcards.Add(TunnelPin);
+				}
+			}
+		}
+	}
+
+	// no wildcards to infer, bail:
+	if(TunnelWildcards.Num() == 0)
+	{
+		return;
+	}
+			
+	// Seed any tunnel wildcard pins that have known values on the macro instance:
+	TArray<TPair<UEdGraphNode*, UEdGraphPin*>> DirtyNodePins; // when a pin is inferred we want to give the node a chance to propagate
+	for (UEdGraphPin* TunnelPin : TunnelWildcards)
+	{
+		UEdGraphPin* MacroPin = FindPin(TunnelPin->GetName(), UEdGraphPin::GetComplementaryDirection(TunnelPin->Direction));
+		if (ensure(MacroPin))
+		{
+			if (!FWildcardNodeUtils::IsWildcardPin(MacroPin))
+			{
+				// tunnel is wildcard, but we are not .. we want to set type on the tunnel and allow inference to run
+				FEdGraphPinType const& LinkedPinType = MacroPin->PinType;
+				FWildcardNodeUtils::InferType(TunnelPin, LinkedPinType);
+
+				for(UEdGraphPin* LinkedPin : TunnelPin->LinkedTo)
+				{
+					if (FWildcardNodeUtils::IsWildcardPin(LinkedPin))
+					{
+						DirtyNodePins.AddUnique({LinkedPin->GetOwningNode(), LinkedPin});
+					}
+				}
+			}
+		}
+	}
+
+	// Helper to count the number of wildcard pins on a node
+	// we monitor these counts to detect when notifications need
+	// to be sent to owning nodes:
+	const auto CountWildcardPins = [](const UEdGraphNode* Node)
+	{
+		int32 WildcardCount = 0;
+		for(UEdGraphPin* Pin : Node->Pins)
+		{
+			if(FWildcardNodeUtils::IsWildcardPin(Pin))
+			{
+				++WildcardCount;
+			}
+		}
+		return WildcardCount;
+	};
+
+	// helper for counting the number of connections a node has,
+	// used to validate that we aren't modifying graph topology
+	const auto CountConnections = [](const UEdGraphNode* Node)
+	{
+		int32 ConnectionCount = 0;
+		for(UEdGraphPin* Pin : Node->Pins)
+		{
+			ConnectionCount += Pin->LinkedTo.Num();
+		}
+		return ConnectionCount;
+	};
+	
+	TMap<UEdGraphNode*, int32> WildcardCounts;
+	TMap<UEdGraphNode*, int32> ConnectionCounts;
+	for(UEdGraphNode* Node : InNodes)
+	{
+		const int32 WildcardCount = CountWildcardPins(Node);
+		if(WildcardCount > 0)
+		{
+			WildcardCounts.Add(Node, WildcardCount);
+		}
+		ConnectionCounts.Add(Node, CountConnections(Node));
+	}
+
+	// We've seeded, now iteratively refresh nodes until pins stabilize:
+	while(DirtyNodePins.Num())
+	{
+		TArray<TPair<UEdGraphNode*, UEdGraphPin*>> CurrentDirtyNodePins = MoveTemp(DirtyNodePins);
+		for(const TPair<UEdGraphNode*, UEdGraphPin*>& DirtyNodePin : CurrentDirtyNodePins)
+		{
+			// Any wildcard pins that are connected to this pin need to be
+			// inferred and marked dirty:
+			CastChecked<UK2Node>(DirtyNodePin.Key)->NotifyPinConnectionListChanged(DirtyNodePin.Value);
+		}
+
+		const auto InferLinkedPins = [](UEdGraphPin * Pin, const UEdGraphPin * SourcePin, TArray<TPair<UEdGraphNode*, UEdGraphPin*>>&OutDirtyNodePins)
+		{
+			TSet<UEdGraphPin*> ProcessedPins;
+			UE::Private::InferLinkedPinsImpl(Pin, SourcePin->PinType, OutDirtyNodePins, ProcessedPins);
+		};
+
+		// look for pins that are now inferable:
+		for(const TPair<UEdGraphNode*, int32>& NodeWithCount : WildcardCounts)
+		{
+			// if count has changed the node is dirty:
+			UEdGraphNode* WildcardNode = NodeWithCount.Key;
+			const int32 WildcardCount = CountWildcardPins(WildcardNode);
+			if(NodeWithCount.Value != WildcardCount)
+			{
+				for(UEdGraphPin* Pin : WildcardNode->Pins)
+				{
+					for(UEdGraphPin* LinkedPin : Pin->LinkedTo)
+					{
+						if(FWildcardNodeUtils::IsWildcardPin(LinkedPin))
+						{
+							// infer and mark dirty:
+							InferLinkedPins(LinkedPin, Pin, DirtyNodePins);
+						}
+					}
+				}
+			}
+
+			// we must also update the count:
+			WildcardCounts[NodeWithCount.Key] = WildcardCount;
+		}
+	}
+
+	for(const TPair<UEdGraphNode*, int32>& NodeWithCount : ConnectionCounts)
+	{
+		// This ensure indicates that we have a node that is destroying the graph in its 
+		// NotifyPinConnectionListChanged override - note that this is an imperfect test
+		// but it seems to be cheap and will catch the most egregious errors
+		UEdGraph* MacroGraph = GetMacroGraph();
+		ensureMsgf(NodeWithCount.Value == CountConnections(NodeWithCount.Key), 
+			TEXT("Node connection count changed while inferring %s - consider setting [Blueprints] bUseSimpleWildcardInference as a workaround"),
+			MacroGraph ? *MacroGraph->GetPathName() : TEXT("Unknown Graph")
+		);
+	}
+
+	// copy back inferred values to any pins on our macro instance:
+	for (const UEdGraphPin* TunnelWildcard : TunnelWildcards)
+	{
+		UEdGraphPin* SourcePin = FindPin(
+			*TunnelWildcard->PinName.ToString(), 
+			UEdGraphPin::GetComplementaryDirection(TunnelWildcard->Direction));
+		if(FWildcardNodeUtils::IsWildcardPin(SourcePin))
+		{
+			FWildcardNodeUtils::InferType(SourcePin, TunnelWildcard->PinType);
+		}
+	}
+}
+
 #undef LOCTEXT_NAMESPACE
diff --git a/Engine/Source/Editor/KismetCompiler/Private/KismetCompiler.cpp b/Engine/Source/Editor/KismetCompiler/Private/KismetCompiler.cpp
--- a/Engine/Source/Editor/KismetCompiler/Private/KismetCompiler.cpp
+++ b/Engine/Source/Editor/KismetCompiler/Private/KismetCompiler.cpp
@@ -4233,27 +4233,6 @@ void FKismetCompilerContext::ExpandTunnelsAndMacros(UEdGraph* SourceGraph)
 
 			TArray<UEdGraphNode*> MacroNodes(ClonedGraph->Nodes);
 
-			// resolve any wildcard pins in the nodes cloned from the macro
-			if (!MacroInstanceNode->ResolvedWildcardType.PinCategory.IsNone())
-			{
-				for (UEdGraphNode* const ClonedNode : ClonedGraph->Nodes)
-				{
-					if (ClonedNode)
-					{
-						for (UEdGraphPin* const ClonedPin : ClonedNode->Pins)
-						{
-							if ( ClonedPin && (ClonedPin->PinType.PinCategory == UEdGraphSchema_K2::PC_Wildcard) )
-							{
-								// copy only type info, so array or ref status is preserved
-								ClonedPin->PinType.PinCategory = MacroInstanceNode->ResolvedWildcardType.PinCategory;
-								ClonedPin->PinType.PinSubCategory = MacroInstanceNode->ResolvedWildcardType.PinSubCategory;
-								ClonedPin->PinType.PinSubCategoryObject = MacroInstanceNode->ResolvedWildcardType.PinSubCategoryObject;
-							}
-						}
-					}
-				}
-			}
-
 			// Handle any nodes that need to inherit their macro instance's NodeGUID
 			for( UEdGraphNode* ClonedNode : MacroNodes)
 			{
@@ -4312,6 +4291,9 @@ void FKismetCompilerContext::ExpandTunnelsAndMacros(UEdGraph* SourceGraph)
 			ClonedGraph->MoveNodesToAnotherGraph(SourceGraph, IsAsyncLoading() || bIsLoading, Blueprint && Blueprint->bBeingCompiled);
 			FEdGraphUtilities::MergeChildrenGraphsIn(SourceGraph, ClonedGraph, /*bRequireSchemaMatch=*/ true);
 
+			// run type inference on MacroNodes:
+			MacroInstanceNode->InferWildcards(MacroNodes);
+
 			// When emitting intermediate products; make an effort to make them readable by preventing overlaps and adding informative comments
 			int32 NodeOffsetX = 0;
 			int32 NodeOffsetY = 0;
diff --git a/Engine/Source/Editor/UnrealEd/Private/Kismet2/WildcardNodeUtils.cpp b/Engine/Source/Editor/UnrealEd/Private/Kismet2/WildcardNodeUtils.cpp
--- a/Engine/Source/Editor/UnrealEd/Private/Kismet2/WildcardNodeUtils.cpp
+++ b/Engine/Source/Editor/UnrealEd/Private/Kismet2/WildcardNodeUtils.cpp
@@ -69,4 +69,41 @@ bool FWildcardNodeUtils::NodeHasAnyWildcards(const UEdGraphNode* const Node)
 	}
 
 	return false;
-}
\ No newline at end of file
+}
+
+void FWildcardNodeUtils::InferType(UEdGraphPin* ToPin, const FEdGraphPinType& Type)
+{
+	InferType(ToPin->PinType, Type);
+}
+
+void FWildcardNodeUtils::InferType(FEdGraphPinType& ToType, const FEdGraphPinType& Type)
+{
+	check(ToType.PinCategory == UEdGraphSchema_K2::PC_Wildcard);
+	ToType.PinCategory = Type.PinCategory;
+	ToType.PinSubCategory = Type.PinSubCategory;
+	ToType.PinSubCategoryObject = Type.PinSubCategoryObject;
+	ToType.PinValueType = Type.PinValueType;
+	ToType.bIsUObjectWrapper = Type.bIsUObjectWrapper;
+}
+
+void FWildcardNodeUtils::ResetToWildcard(UEdGraphPin* Pin)
+{
+	const UEdGraphSchema_K2* Schema = GetDefault<UEdGraphSchema_K2>();
+	check(Pin);
+	if (Pin->SubPins.Num() > 0)
+	{
+		Schema->RecombinePin(Pin->SubPins[0]);
+	}
+
+	ResetToWildcard(Pin->PinType);
+	Schema->ResetPinToAutogeneratedDefaultValue(Pin, false);
+}
+
+void FWildcardNodeUtils::ResetToWildcard(FEdGraphPinType& PinType)
+{
+	PinType.PinCategory = UEdGraphSchema_K2::PC_Wildcard;
+	PinType.PinSubCategory = NAME_None;
+	PinType.PinSubCategoryObject = nullptr;
+	PinType.bIsUObjectWrapper = false;
+}
+
diff --git a/Engine/Source/Editor/UnrealEd/Public/Kismet2/WildcardNodeUtils.h b/Engine/Source/Editor/UnrealEd/Public/Kismet2/WildcardNodeUtils.h
--- a/Engine/Source/Editor/UnrealEd/Public/Kismet2/WildcardNodeUtils.h
+++ b/Engine/Source/Editor/UnrealEd/Public/Kismet2/WildcardNodeUtils.h
@@ -53,4 +53,16 @@ public:
 	* @return	True if the given node has any wildcard pins on it
 	*/
 	static UNREALED_API bool NodeHasAnyWildcards(const UEdGraphNode* const Node);
+	
+	/** 
+	 * Utility functions for overwriting a wildcard type with an inferred type. 
+	 * These will leave the pin's container information unchanged - e.g. so that
+	 * a macro node can take in an 'array of wildcards'. 
+	 */
+	static UNREALED_API void InferType(UEdGraphPin* ToPin, const FEdGraphPinType& Type);
+	static UNREALED_API void InferType(FEdGraphPinType& ToType, const FEdGraphPinType& Type);
+
+	/** Utility functions for resetting a pin to wildcard, again, leaving container information unchanged */
+	static UNREALED_API void ResetToWildcard(UEdGraphPin* Pin);
+	static UNREALED_API void ResetToWildcard(FEdGraphPinType& PinType);
 };
