Disable smart inference while I work through issues with partial wildcards (e.g. TMaps with only wildcard keys or values)

#rnx

diff --git a/Engine/Source/Editor/BlueprintGraph/Private/K2Node_MacroInstance.cpp b/Engine/Source/Editor/BlueprintGraph/Private/K2Node_MacroInstance.cpp
--- a/Engine/Source/Editor/BlueprintGraph/Private/K2Node_MacroInstance.cpp
+++ b/Engine/Source/Editor/BlueprintGraph/Private/K2Node_MacroInstance.cpp
@@ -630,8 +630,13 @@ TArray<UEdGraphPin*> UK2Node_MacroInstance::GetAllWildcardPins() const
 
 bool UK2Node_MacroInstance::ShouldDoSmartWildcardInference()
 {
+#if 0
 	static const FBoolConfigValueHelper bUseSimpleWildcardInference(TEXT("Blueprints"), TEXT("bUseSimpleWildcardInference"), GEngineIni);
 	return !bUseSimpleWildcardInference;
+#else
+	// disabled until issues with partial wildcard types (e.g. TMaps with only wildcard keys or values) are fixed
+	return false;
+#endif
 }
 
 namespace UE::Private
