Add the missing "Replicated" node title to K2Node_TImeline

#rnx

diff --git a/Engine/Source/Editor/BlueprintGraph/Classes/K2Node_Timeline.h b/Engine/Source/Editor/BlueprintGraph/Classes/K2Node_Timeline.h
--- a/Engine/Source/Editor/BlueprintGraph/Classes/K2Node_Timeline.h
+++ b/Engine/Source/Editor/BlueprintGraph/Classes/K2Node_Timeline.h
@@ -84,6 +84,7 @@ class UK2Node_Timeline : public UK2Node
 	virtual void ExpandNode(class FKismetCompilerContext& CompilerContext, UEdGraph* SourceGraph) override;
 	virtual void GetNodeAttributes( TArray<TKeyValuePair<FString, FString>>& OutNodeAttributes ) const override;
 	virtual void GetMenuActions(FBlueprintActionDatabaseRegistrar& ActionRegistrar) const override;
+	virtual FText GetToolTipHeading() const override;
 	//~ End UK2Node Interface.
 
 	/** Get the 'play' input pin */
diff --git a/Engine/Source/Editor/BlueprintGraph/Private/K2Node_Timeline.cpp b/Engine/Source/Editor/BlueprintGraph/Private/K2Node_Timeline.cpp
--- a/Engine/Source/Editor/BlueprintGraph/Private/K2Node_Timeline.cpp
+++ b/Engine/Source/Editor/BlueprintGraph/Private/K2Node_Timeline.cpp
@@ -661,6 +661,29 @@ FText UK2Node_Timeline::GetTooltipText() const
 	return LOCTEXT("TimelineTooltip", "Timeline node allows values to be keyframed over time.\nDouble click to open timeline editor.");
 }
 
+FText UK2Node_Timeline::GetToolTipHeading() const
+{
+	const FText ExistingHeading = Super::GetToolTipHeading();
+	
+	if (bReplicated)
+	{
+		const FText ReplicatedLabel = LOCTEXT("ReplicatedTimelineComponent", "Replicated");
+
+		// If the existing heading is empty then we don't need anything other then "Replicated"
+		if (ExistingHeading.IsEmpty())
+		{
+			return ReplicatedLabel;
+		}
+		// otherwise, we need an extra newline character ( '\n' ) to get everything properly formatted 
+		else
+		{
+			return FText::Format(LOCTEXT("UK2Node_Timeline_GetToolTipHeading", "{0}\n{1}"), {ReplicatedLabel, ExistingHeading });	
+		}
+	}
+
+	return ExistingHeading;
+}
+
 FName UK2Node_Timeline::GetCornerIcon() const
 {
 	if (bReplicated)
