More fixes for 'smart' macro wildcard inference, specifically around TMap value parameters, which historically have been left in an unknown 'wildcard' state until compilation
#rnx

diff --git a/Engine/Source/Editor/BlueprintGraph/Classes/K2Node_MacroInstance.h b/Engine/Source/Editor/BlueprintGraph/Classes/K2Node_MacroInstance.h
--- a/Engine/Source/Editor/BlueprintGraph/Classes/K2Node_MacroInstance.h
+++ b/Engine/Source/Editor/BlueprintGraph/Classes/K2Node_MacroInstance.h
@@ -119,7 +119,6 @@ private:
 	// the graph, and then back to any input pins:
 	void InferWildcards();
 	TArray<UEdGraphPin*> GetAllWildcardPins() const;
-	static bool ShouldDoSmartWildcardInference();
 	void SmartInferWildcardsImpl(const TArray<UEdGraphNode*>& InNodes) const;
 
 	/** Constructing FText strings can be costly, so we cache the node's tooltip */
diff --git a/Engine/Source/Editor/BlueprintGraph/Classes/K2Node_Tunnel.h b/Engine/Source/Editor/BlueprintGraph/Classes/K2Node_Tunnel.h
--- a/Engine/Source/Editor/BlueprintGraph/Classes/K2Node_Tunnel.h
+++ b/Engine/Source/Editor/BlueprintGraph/Classes/K2Node_Tunnel.h
@@ -95,6 +95,9 @@ protected:
 	* @param bInAllWildcardPinsUnlinked	TRUE if all wildcard pins were unlinked
 	*/
 	virtual void PostFixupAllWildcardPins(bool bInAllWildcardPinsUnlinked) {}
+	// Feature flag for 'smart' wildcard inference which allows for multiple 
+	// types to be inferred.
+	static bool ShouldDoSmartWildcardInference();
 
 	// Utility function that subclasses must call after allocating their default pins
 	void CacheWildcardPins();
diff --git a/Engine/Source/Editor/BlueprintGraph/Private/K2Node_CallFunction.cpp b/Engine/Source/Editor/BlueprintGraph/Private/K2Node_CallFunction.cpp
--- a/Engine/Source/Editor/BlueprintGraph/Private/K2Node_CallFunction.cpp
+++ b/Engine/Source/Editor/BlueprintGraph/Private/K2Node_CallFunction.cpp
@@ -25,6 +25,7 @@
 #include "K2Node_TemporaryVariable.h"
 #include "K2Node_ExecutionSequence.h"
 #include "Kismet2/BlueprintEditorUtils.h"
+#include "Kismet2/WildcardNodeUtils.h"
 #include "Settings/EditorStyleSettings.h"
 #include "Editor.h"
 #include "Editor/UnrealEdEngine.h"
@@ -3090,14 +3091,20 @@ void UK2Node_CallFunction::ConformContainerPins()
 		{
 			if (Pin->HasAnyConnections() || !Pin->DoesDefaultValueMatchAutogenerated() )
 			{
-				bOutPropagated = true;
+				FEdGraphTerminalType TypeToPotentiallyPropagate;
 				if (Pin->LinkedTo.Num() != 0)
 				{
-					TypeToPropagete = Pin->LinkedTo[0]->GetPrimaryTerminalType();
+					TypeToPotentiallyPropagate = Pin->LinkedTo[0]->GetPrimaryTerminalType();
 				}
 				else
 				{
-					TypeToPropagete = Pin->GetPrimaryTerminalType();
+					TypeToPotentiallyPropagate = Pin->GetPrimaryTerminalType();
+				}
+				
+				if (!FWildcardNodeUtils::IsWildcardPin(TypeToPotentiallyPropagate))
+				{
+					bOutPropagated = true;
+					TypeToPropagete = TypeToPotentiallyPropagate;
 				}
 			}
 		}
@@ -3109,14 +3116,16 @@ void UK2Node_CallFunction::ConformContainerPins()
 		{
 			if (Pin->LinkedTo.Num() != 0 || !Pin->DoesDefaultValueMatchAutogenerated())
 			{
-				bOutPropagated = true;
+				const FEdGraphTerminalType* ValueType = &Pin->PinType.PinValueType;
 				if (Pin->LinkedTo.Num() != 0)
 				{
-					TypeToPropagete = Pin->LinkedTo[0]->PinType.PinValueType;
+					ValueType = &(Pin->LinkedTo[0]->PinType.PinValueType);
 				}
-				else
+
+				if (!FWildcardNodeUtils::IsWildcardPin(*ValueType))
 				{
-					TypeToPropagete = Pin->PinType.PinValueType;
+					bOutPropagated = true;
+					TypeToPropagete = *ValueType;
 				}
 			}
 		}
diff --git a/Engine/Source/Editor/BlueprintGraph/Private/K2Node_Knot.cpp b/Engine/Source/Editor/BlueprintGraph/Private/K2Node_Knot.cpp
--- a/Engine/Source/Editor/BlueprintGraph/Private/K2Node_Knot.cpp
+++ b/Engine/Source/Editor/BlueprintGraph/Private/K2Node_Knot.cpp
@@ -8,6 +8,7 @@
 #include "EdGraphSchema_K2.h"
 #include "Internationalization/Internationalization.h"
 #include "Kismet2/Kismet2NameValidators.h"
+#include "Kismet2/WildcardNodeUtils.h"
 #include "Misc/AssertionMacros.h"
 #include "Templates/Casts.h"
 #include "Templates/UnrealTemplate.h"
@@ -98,7 +99,7 @@ void UK2Node_Knot::PropagatePinType()
 
 	for (UEdGraphPin* Inputs : MyInputPin->LinkedTo)
 	{
-		if (Inputs->PinType.PinCategory != UEdGraphSchema_K2::PC_Wildcard)
+		if (!FWildcardNodeUtils::HasAnyWildcards(Inputs))
 		{
 			PropagatePinTypeFromDirection(true);
 			return;
@@ -107,7 +108,7 @@ void UK2Node_Knot::PropagatePinType()
 
 	for (UEdGraphPin* Outputs : MyOutputPin->LinkedTo)
 	{
-		if (Outputs->PinType.PinCategory != UEdGraphSchema_K2::PC_Wildcard)
+		if (!FWildcardNodeUtils::HasAnyWildcards(Outputs))
 		{
 			PropagatePinTypeFromDirection(false);
 			return;
diff --git a/Engine/Source/Editor/BlueprintGraph/Private/K2Node_MacroInstance.cpp b/Engine/Source/Editor/BlueprintGraph/Private/K2Node_MacroInstance.cpp
--- a/Engine/Source/Editor/BlueprintGraph/Private/K2Node_MacroInstance.cpp
+++ b/Engine/Source/Editor/BlueprintGraph/Private/K2Node_MacroInstance.cpp
@@ -284,26 +284,36 @@ void UK2Node_MacroInstance::FindInContentBrowser(TWeakObjectPtr<UK2Node_MacroIns
 void UK2Node_MacroInstance::NotifyPinConnectionListChanged(UEdGraphPin* ChangedPin)
 {
 	Super::NotifyPinConnectionListChanged(ChangedPin);
+	const bool bShouldDoSmartInference = ShouldDoSmartWildcardInference();
+	if(bShouldDoSmartInference)
+	{
+		const bool bIsWildcardPin = FWildcardNodeUtils::HasAnyWildcards(ChangedPin);
+		if (bIsWildcardPin && ChangedPin->LinkedTo.Num() > 0)
+		{
+			// Search the changed pin's links for an inferrable pin:
+			if(const UEdGraphPin* InferrablePin = FWildcardNodeUtils::FindInferrableLinkedPin(ChangedPin))
+			{
+				// we found one, infer from it and then propagate the inference:
+				FWildcardNodeUtils::InferType(ChangedPin, InferrablePin->PinType);
+				InferWildcards();
+			}
+		}
+	}
 
 	// added a link?
 	if (ChangedPin->LinkedTo.Num() > 0)
 	{
 		// ... to a wildcard pin?
-		bool const bIsWildcardPin = ChangedPin->PinType.PinCategory == UEdGraphSchema_K2::PC_Wildcard;
+		const bool bIsWildcardPin = ChangedPin->PinType.PinCategory == UEdGraphSchema_K2::PC_Wildcard;
 		if (bIsWildcardPin)
 		{
 			// get type of pin we just got linked to
 			FEdGraphPinType const LinkedPinType = ChangedPin->LinkedTo[0]->PinType;
-			if(ShouldDoSmartWildcardInference())
-			{
-				// only copy the category stuff to preserve array and ref status
-				FWildcardNodeUtils::InferType(ChangedPin, LinkedPinType);
-				InferWildcards();
-			}
-			else
+
+			// change all other wildcard pins to the new type
+			// note we're assuming only one wildcard type per Macro node, for now
+			if(!bShouldDoSmartInference)
 			{
-				// change all other wildcard pins to the new type
-				// note we're assuming only one wildcard type per Macro node, for now
 				for(int32 PinIdx=0; PinIdx<Pins.Num(); PinIdx++)
 				{
 					UEdGraphPin* const TmpPin = Pins[PinIdx];
@@ -363,6 +373,33 @@ void UK2Node_MacroInstance::PostReconstructNode()
 {
 	bReconstructNode = false;
 
+	if(ShouldDoSmartWildcardInference())
+	{
+		// conform any type mismatches - or just conform
+		for(UEdGraphPin* Pin : WildcardPins)
+		{
+			if (FWildcardNodeUtils::HasAnyWildcards(Pin))
+			{
+				const FEdGraphPinType* ConnectedType = nullptr;
+				for(UEdGraphPin* Link : Pin->LinkedTo)
+				{
+					if(!FWildcardNodeUtils::HasAnyWildcards(Link))
+					{
+						ConnectedType = &Link->PinType;
+					}
+				}
+
+				if(ConnectedType)
+				{
+					FWildcardNodeUtils::InferType(Pin->PinType, *ConnectedType);
+				}
+			}
+		}
+
+		// rerun inference
+		InferWildcards();
+	}
+
 	Super::PostReconstructNode();
 }
 
@@ -474,12 +511,20 @@ void UK2Node_MacroInstance::PostFixupAllWildcardPins(bool bInAllWildcardPinsUnli
 		ResolvedWildcardType.ResetToDefaults();
 
 		// Collapse any wildcard pins that are split and set their type back to wildcard
-		for (UEdGraphPin* Pin : WildcardPins)
+		// doing this would be unsafe when using smart wildcard inference
+		// because recombining pin in the middle of reconstruction could result in
+		// pin allocation during reconstruction. Therefore we don't rely upon it 
+		// when doing smart wildcard inference
+		if (!ShouldDoSmartWildcardInference())
 		{
-			GetSchema()->RecombinePin(Pin);
-			Pin->PinType.PinCategory = UEdGraphSchema_K2::PC_Wildcard;
-			Pin->PinType.PinSubCategory = NAME_None;
-			Pin->PinType.PinSubCategoryObject = nullptr;
+			for (UEdGraphPin* Pin : WildcardPins)
+			{
+				GetSchema()->RecombinePin(Pin);
+
+				Pin->PinType.PinCategory = UEdGraphSchema_K2::PC_Wildcard;
+				Pin->PinType.PinSubCategory = NAME_None;
+				Pin->PinType.PinSubCategoryObject = nullptr;
+			}
 		}
 	}
 }
@@ -628,17 +673,6 @@ TArray<UEdGraphPin*> UK2Node_MacroInstance::GetAllWildcardPins() const
 	return Result;
 }
 
-bool UK2Node_MacroInstance::ShouldDoSmartWildcardInference()
-{
-#if 0
-	static const FBoolConfigValueHelper bUseSimpleWildcardInference(TEXT("Blueprints"), TEXT("bUseSimpleWildcardInference"), GEngineIni);
-	return !bUseSimpleWildcardInference;
-#else
-	// disabled until issues with partial wildcard types (e.g. TMaps with only wildcard keys or values) are fixed
-	return false;
-#endif
-}
-
 namespace UE::Private
 {
 
@@ -670,7 +704,7 @@ void UK2Node_MacroInstance::SmartInferWildcardsImpl(const TArray<UEdGraphNode*>&
 		{
 			for (UEdGraphPin* TunnelPin : Tunnel->Pins)
 			{
-				if (FWildcardNodeUtils::IsWildcardPin(TunnelPin))
+				if (FWildcardNodeUtils::HasAnyWildcards(TunnelPin))
 				{
 					// the tunnel node with input pins is the output on the macro:
 					TunnelWildcards.Add(TunnelPin);
@@ -700,7 +734,7 @@ void UK2Node_MacroInstance::SmartInferWildcardsImpl(const TArray<UEdGraphNode*>&
 
 				for(UEdGraphPin* LinkedPin : TunnelPin->LinkedTo)
 				{
-					if (FWildcardNodeUtils::IsWildcardPin(LinkedPin))
+					if (FWildcardNodeUtils::HasAnyWildcards(LinkedPin))
 					{
 						DirtyNodePins.AddUnique({LinkedPin->GetOwningNode(), LinkedPin});
 					}
@@ -717,7 +751,7 @@ void UK2Node_MacroInstance::SmartInferWildcardsImpl(const TArray<UEdGraphNode*>&
 		int32 WildcardCount = 0;
 		for(UEdGraphPin* Pin : Node->Pins)
 		{
-			if(FWildcardNodeUtils::IsWildcardPin(Pin))
+			if(FWildcardNodeUtils::HasAnyWildcards(Pin))
 			{
 				++WildcardCount;
 			}
@@ -778,7 +812,7 @@ void UK2Node_MacroInstance::SmartInferWildcardsImpl(const TArray<UEdGraphNode*>&
 				{
 					for(UEdGraphPin* LinkedPin : Pin->LinkedTo)
 					{
-						if(FWildcardNodeUtils::IsWildcardPin(LinkedPin))
+						if(FWildcardNodeUtils::HasAnyWildcards(LinkedPin))
 						{
 							// infer and mark dirty:
 							InferLinkedPins(LinkedPin, Pin, DirtyNodePins);
@@ -810,7 +844,7 @@ void UK2Node_MacroInstance::SmartInferWildcardsImpl(const TArray<UEdGraphNode*>&
 		UEdGraphPin* SourcePin = FindPin(
 			*TunnelWildcard->PinName.ToString(), 
 			UEdGraphPin::GetComplementaryDirection(TunnelWildcard->Direction));
-		if(FWildcardNodeUtils::IsWildcardPin(SourcePin))
+		if(FWildcardNodeUtils::HasAnyWildcards(SourcePin))
 		{
 			FWildcardNodeUtils::InferType(SourcePin, TunnelWildcard->PinType);
 		}
diff --git a/Engine/Source/Editor/BlueprintGraph/Private/K2Node_Tunnel.cpp b/Engine/Source/Editor/BlueprintGraph/Private/K2Node_Tunnel.cpp
--- a/Engine/Source/Editor/BlueprintGraph/Private/K2Node_Tunnel.cpp
+++ b/Engine/Source/Editor/BlueprintGraph/Private/K2Node_Tunnel.cpp
@@ -12,6 +12,7 @@
 #include "K2Node_Composite.h"
 #include "K2Node_MacroInstance.h"
 #include "Kismet2/BlueprintEditorUtils.h"
+#include "Kismet2/WildcardNodeUtils.h"
 #include "Math/UnrealMathSSE.h"
 #include "Templates/Casts.h"
 #include "Templates/Function.h"
@@ -277,6 +278,12 @@ bool UK2Node_Tunnel::ModifyUserDefinedPinDefaultValue(TSharedPtr<FUserPinInfo> P
 	return false;
 }
 
+bool UK2Node_Tunnel::ShouldDoSmartWildcardInference()
+{
+	static const FBoolConfigValueHelper bUseSimpleWildcardInference(TEXT("Blueprints"), TEXT("bUseSimpleWildcardInference"), GEngineIni);
+	return !bUseSimpleWildcardInference;
+}
+
 bool UK2Node_Tunnel::CanModifyExecutionWires()
 {
 	return true;
@@ -321,12 +328,26 @@ void UK2Node_Tunnel::CacheWildcardPins()
 {
 	WildcardPins.Reset();
 
-	for (UEdGraphPin* Pin : Pins)
+	const bool bSupportMapWildcards = ShouldDoSmartWildcardInference();
+	if(bSupportMapWildcards)
+	{
+		for (UEdGraphPin* Pin : Pins)
+		{
+			if (FWildcardNodeUtils::HasAnyWildcards(Pin))
+			{
+				WildcardPins.Add(Pin);
+			}
+		}
+	}
+	else
 	{
-		// for each of the wildcard pins...
-		if (Pin->PinType.PinCategory == UEdGraphSchema_K2::PC_Wildcard)
+		for (UEdGraphPin* Pin : Pins)
 		{
-			WildcardPins.Add(Pin);
+			// for each of the wildcard pins...
+			if (Pin->PinType.PinCategory == UEdGraphSchema_K2::PC_Wildcard)
+			{
+				WildcardPins.Add(Pin);
+			}
 		}
 	}
 }
@@ -376,11 +397,13 @@ void UK2Node_Tunnel::ReallocatePinsDuringReconstruction(TArray<UEdGraphPin*>& Ol
 
 	if (bAllWildcardsAreUnlinked == false)
 	{
+		const bool bSupportMapWildcards = ShouldDoSmartWildcardInference();
 		// Copy pin types from old pins for wildcard pins
 		for (UEdGraphPin* const Pin : WildcardPins)
 		{
 			// Only change the type if it is still a wildcard
-			if (Pin->PinType.PinCategory == UEdGraphSchema_K2::PC_Wildcard)
+			if ((bSupportMapWildcards && FWildcardNodeUtils::HasAnyWildcards(Pin)) ||
+				Pin->PinType.PinCategory == UEdGraphSchema_K2::PC_Wildcard)
 			{
 				// find it in the old pins and copy the type
 				for (UEdGraphPin const* const OldPin : OldPins)
diff --git a/Engine/Source/Editor/UnrealEd/Private/Kismet2/WildcardNodeUtils.cpp b/Engine/Source/Editor/UnrealEd/Private/Kismet2/WildcardNodeUtils.cpp
--- a/Engine/Source/Editor/UnrealEd/Private/Kismet2/WildcardNodeUtils.cpp
+++ b/Engine/Source/Editor/UnrealEd/Private/Kismet2/WildcardNodeUtils.cpp
@@ -2,6 +2,7 @@
 
 #include "Kismet2/WildcardNodeUtils.h"
 
+#include "Algo/Find.h"
 #include "Containers/Array.h"
 #include "EdGraphSchema_K2.h"
 #include "Misc/AssertionMacros.h"
@@ -27,6 +28,47 @@ bool FWildcardNodeUtils::IsWildcardPin(const UEdGraphPin* const Pin)
 	return Pin && Pin->PinType.PinCategory == UEdGraphSchema_K2::PC_Wildcard;
 }
 
+bool FWildcardNodeUtils::IsWildcardPin(const FEdGraphTerminalType& Terminal)
+{
+	return Terminal.TerminalCategory == UEdGraphSchema_K2::PC_Wildcard;
+}
+
+bool FWildcardNodeUtils::HasAnyWildcards(const UEdGraphPin* const Pin)
+{
+	check(Pin);
+	return Pin->PinType.PinCategory == UEdGraphSchema_K2::PC_Wildcard ||
+		Pin->PinType.PinValueType.TerminalCategory == UEdGraphSchema_K2::PC_Wildcard ;
+}
+
+bool FWildcardNodeUtils::HasAnyNonWildcards(const UEdGraphPin* Pin)
+{
+	check(Pin);
+	return Pin->PinType.PinCategory != UEdGraphSchema_K2::PC_Wildcard ||
+		Pin->PinType.PinValueType.TerminalCategory != UEdGraphSchema_K2::PC_Wildcard;
+}
+
+const UEdGraphPin* FWildcardNodeUtils::FindInferrableLinkedPin( const UEdGraphPin* ForPin)
+{
+	// first, look for pins with no wildcards to infer from:
+	UEdGraphPin *const* InferrablePin = Algo::FindByPredicate(ForPin->LinkedTo, 
+		[](const UEdGraphPin* Link)
+		{
+			return !FWildcardNodeUtils::HasAnyWildcards(Link);
+		});
+	if(InferrablePin == nullptr)
+	{
+		// no pins with no wildcards - are there any pins that have
+		// any non wildcards?
+		InferrablePin = Algo::FindByPredicate(ForPin->LinkedTo, 
+			[](const UEdGraphPin* Link)
+			{
+				return FWildcardNodeUtils::HasAnyNonWildcards(Link);
+			});
+	}
+
+	return InferrablePin ? *InferrablePin : nullptr;
+}
+
 bool FWildcardNodeUtils::IsLinkedToWildcard(const UEdGraphPin* const Pin)
 {
 	if (Pin)
@@ -78,12 +120,21 @@ void FWildcardNodeUtils::InferType(UEdGraphPin* ToPin, const FEdGraphPinType& Ty
 
 void FWildcardNodeUtils::InferType(FEdGraphPinType& ToType, const FEdGraphPinType& Type)
 {
-	check(ToType.PinCategory == UEdGraphSchema_K2::PC_Wildcard);
-	ToType.PinCategory = Type.PinCategory;
-	ToType.PinSubCategory = Type.PinSubCategory;
-	ToType.PinSubCategoryObject = Type.PinSubCategoryObject;
-	ToType.PinValueType = Type.PinValueType;
-	ToType.bIsUObjectWrapper = Type.bIsUObjectWrapper;
+	bool bInferredSomething = false;
+	if(ToType.PinCategory == UEdGraphSchema_K2::PC_Wildcard)
+	{
+		bInferredSomething = true;
+		ToType.PinCategory = Type.PinCategory;
+		ToType.PinSubCategory = Type.PinSubCategory;
+		ToType.PinSubCategoryObject = Type.PinSubCategoryObject;
+		ToType.bIsUObjectWrapper = Type.bIsUObjectWrapper;
+	}
+	if(ToType.PinValueType.TerminalCategory == UEdGraphSchema_K2::PC_Wildcard)
+	{
+		bInferredSomething = true;
+		ToType.PinValueType = Type.PinValueType;
+	}
+	check(bInferredSomething);
 }
 
 void FWildcardNodeUtils::ResetToWildcard(UEdGraphPin* Pin)
diff --git a/Engine/Source/Editor/UnrealEd/Public/Kismet2/WildcardNodeUtils.h b/Engine/Source/Editor/UnrealEd/Public/Kismet2/WildcardNodeUtils.h
--- a/Engine/Source/Editor/UnrealEd/Public/Kismet2/WildcardNodeUtils.h
+++ b/Engine/Source/Editor/UnrealEd/Public/Kismet2/WildcardNodeUtils.h
@@ -20,6 +20,24 @@ public:
 	* and checks during compilation of nodes.
 	*/
 	static UNREALED_API FEdGraphPinType GetDefaultWildcardPinType();
+	
+	/** 
+	 * Checks if the pin has any wildcard components in its type. IsWildcardPin( const UEdGraphPin*) 
+	 * will return false for pins that have only wildcard TMap Values, this will not.
+	 * This function or IsWildcardPin(const FEdGraphTerminalType&) should be preferred.
+	 */
+	static UNREALED_API bool HasAnyWildcards(const UEdGraphPin* Pin);
+
+	/**
+	 * Returns true if the pin has any non wild cards (primary or secondary type)
+	 */
+	static UNREALED_API bool HasAnyNonWildcards(const UEdGraphPin* Pin);
+
+	/**
+	 * Returns a non wildcard pin from ForPin's LinkedTo list. Preferring pins
+	 * with no wildcards to pins with some wildcards.
+	 */
+	static UNREALED_API const UEdGraphPin* FindInferrableLinkedPin( const UEdGraphPin* ForPin);
 
 	/**
 	* Checks if the given pin is in a wildcard state
@@ -28,6 +46,7 @@ public:
 	* @return	True if the given pin is a Wildcard pin
 	*/
 	static UNREALED_API bool IsWildcardPin(const UEdGraphPin* const Pin);
+	static UNREALED_API bool IsWildcardPin(const FEdGraphTerminalType& Terminal);
 
 	/**
 	* Checks if the given pin is linked to any wildcard pins
