Disable smart macro inference while I work on a solution for delaying inference to allow for making macros that delay inference

#rnx

diff --git a/Engine/Source/Editor/BlueprintGraph/Private/K2Node_Tunnel.cpp b/Engine/Source/Editor/BlueprintGraph/Private/K2Node_Tunnel.cpp
--- a/Engine/Source/Editor/BlueprintGraph/Private/K2Node_Tunnel.cpp
+++ b/Engine/Source/Editor/BlueprintGraph/Private/K2Node_Tunnel.cpp
@@ -280,8 +280,15 @@ bool UK2Node_Tunnel::ModifyUserDefinedPinDefaultValue(TSharedPtr<FUserPinInfo> P
 
 bool UK2Node_Tunnel::ShouldDoSmartWildcardInference()
 {
+#if 0
 	static const FBoolConfigValueHelper bUseSimpleWildcardInference(TEXT("Blueprints"), TEXT("bUseSimpleWildcardInference"), GEngineIni);
 	return !bUseSimpleWildcardInference;
+#else
+	// This cannot be enabled until some leniency is added - e.g. to allow
+	// TMaps with a more restrictive value type to supercede inferences of base 
+	// type functionality
+	return false;
+#endif
 }
 
 bool UK2Node_Tunnel::CanModifyExecutionWires()
