Fixup ResolvedWildcardType when we have wildcardpins but it is not set, and we're using simple inference.

#rnx

diff --git a/Engine/Source/Editor/BlueprintGraph/Private/K2Node_MacroInstance.cpp b/Engine/Source/Editor/BlueprintGraph/Private/K2Node_MacroInstance.cpp
--- a/Engine/Source/Editor/BlueprintGraph/Private/K2Node_MacroInstance.cpp
+++ b/Engine/Source/Editor/BlueprintGraph/Private/K2Node_MacroInstance.cpp
@@ -399,6 +399,23 @@ void UK2Node_MacroInstance::PostReconstructNode()
 		// rerun inference
 		InferWildcards();
 	}
+	else
+	{
+		// fix up ResolvedWildcardType, which could have been cleared for certain CL ranges
+		if (ResolvedWildcardType.PinCategory.IsNone() && WildcardPins.Num() > 0)
+		{
+			UEdGraphPin* const* NonWildcardPin = Algo::FindByPredicate(WildcardPins,
+				[](const UEdGraphPin* Pin )
+				{
+					return !FWildcardNodeUtils::IsWildcardPin(Pin);
+				});
+
+			if(NonWildcardPin)
+			{
+				ResolvedWildcardType = (*NonWildcardPin)->PinType;
+			}
+		}
+	}
 
 	Super::PostReconstructNode();
 }
