UE-196156 - Add exec toggle for pure nodes

Blueprint pure nodes were initially intended to be similar to "functional pure". In other words, they're deterministic and produce no side effects (eg: doesn't mutate state). However, pure nodes have violated both conditions for a while now.

Instead, pure nodes are simply function nodes with output values and no visible exec pins. While this can be convenient, they come with a downside: the pure node is evaluated for each connected output. This can lead to unexpected performance issues if the node is expensive to evaluate. In the case of non-deterministic nodes, this can lead to unexpected behavior. In both cases, the user often needs to cull multiple outputs of a pure node and cache the result manually.

The solution here is to add support for toggling purity at the call site. When a function node is placed, right-clicking on it and selecting either "Hide Exec pins" or "Show Exec Pins" will toggle purity. Additionally, the meaning of BlueprintPure and the "Pure" check box changes slightly: it now means that the function node _defaults_ to a pure state when placed in a graph. However, it can be toggled to show its exec pins.

In future changes, we'll also reevaluate which common library functions should continue to default as pure.

diff --git a/Engine/Config/BaseEngine.ini b/Engine/Config/BaseEngine.ini
--- a/Engine/Config/BaseEngine.ini
+++ b/Engine/Config/BaseEngine.ini
@@ -1286,6 +1286,8 @@ DefaultWorldPartitionHLODLayer=/Engine/Maps/Templates/HLODs/HLODLayer_Instanced.
 
 +EnumRedirects=(OldName="/Script/Engine.EUserDefinedStructureStatus",NewName="/Script/CoreUObject.EUserDefinedStructureStatus")
 
++PropertyRedirects=(OldName="K2Node_CallFunction.bIsPureFunc", NewName="bDefaultsToPureFunc")
+
 ; End of releases redirects
 
 [CoreUObject.Metadata]
diff --git a/Engine/Plugins/EnhancedInput/Source/InputBlueprintNodes/Private/K2Node_InputActionValueAccessor.h b/Engine/Plugins/EnhancedInput/Source/InputBlueprintNodes/Private/K2Node_InputActionValueAccessor.h
--- a/Engine/Plugins/EnhancedInput/Source/InputBlueprintNodes/Private/K2Node_InputActionValueAccessor.h
+++ b/Engine/Plugins/EnhancedInput/Source/InputBlueprintNodes/Private/K2Node_InputActionValueAccessor.h
@@ -25,6 +25,11 @@ public:
 
 	void Initialize(const UInputAction* Action);
 
+protected:
+	//~ Begin UK2Node_CallFunction interface
+	virtual bool CanToggleNodePurity() const override { return false; }
+	//~ End UK2Node_CallFunction interface
+
 private:
 	UPROPERTY()
 	TObjectPtr<const UInputAction> InputAction;
diff --git a/Engine/Plugins/Experimental/PythonScriptPlugin/Source/PythonScriptPlugin/Private/K2Node_ExecutePythonScript.h b/Engine/Plugins/Experimental/PythonScriptPlugin/Source/PythonScriptPlugin/Private/K2Node_ExecutePythonScript.h
--- a/Engine/Plugins/Experimental/PythonScriptPlugin/Source/PythonScriptPlugin/Private/K2Node_ExecutePythonScript.h
+++ b/Engine/Plugins/Experimental/PythonScriptPlugin/Source/PythonScriptPlugin/Private/K2Node_ExecutePythonScript.h
@@ -39,6 +39,11 @@ class UK2Node_ExecutePythonScript : public UK2Node_CallFunction
 	virtual int32 GetNodeRefreshPriority() const override { return EBaseNodeRefreshPriority::Low_UsesDependentWildcard; }
 	//~ End UK2Node Interface.
 
+protected:
+	//~ UK2Node_CallFunction interface
+	virtual bool CanToggleNodePurity() const override { return false; }
+	//~ End UK2Node_CallFunction interface
+
 private:
 	/** Synchronize the type of the given argument pin with the type its connected to, or reset it to a wildcard pin if there's no connection */
 	void SynchronizeArgumentPinType(UEdGraphPin* Pin);
@@ -48,7 +53,6 @@ private:
 	UEdGraphPin* FindArgumentPin(const FName PinName, EEdGraphPinDirection PinDirection = EGPD_MAX);
 	UEdGraphPin* FindArgumentPinChecked(const FName PinName, EEdGraphPinDirection PinDirection = EGPD_MAX);
 
-private:
 	/** User-defined input pins */
 	UPROPERTY(EditAnywhere, Category="Arguments")
 	TArray<FName> Inputs;
diff --git a/Engine/Source/Editor/AnimGraph/Public/K2Node_AnimGetter.h b/Engine/Source/Editor/AnimGraph/Public/K2Node_AnimGetter.h
--- a/Engine/Source/Editor/AnimGraph/Public/K2Node_AnimGetter.h
+++ b/Engine/Source/Editor/AnimGraph/Public/K2Node_AnimGetter.h
@@ -137,4 +137,8 @@ protected:
 
 	/** Generates a title for the node based on its function and the context it is in */
 	static FText GenerateTitle(UFunction* Getter, UAnimStateNodeBase* SourceStateNode, UAnimGraphNode_Base* SourceNode);
+
+	// UK2Node_CallFunction interface
+	virtual bool CanToggleNodePurity() const override { return false; }
+	// End of UK2Node_CallFunction interface
 };
diff --git a/Engine/Source/Editor/AnimationBlueprintLibrary/Public/K2_AnimationAttributeNodes.h b/Engine/Source/Editor/AnimationBlueprintLibrary/Public/K2_AnimationAttributeNodes.h
--- a/Engine/Source/Editor/AnimationBlueprintLibrary/Public/K2_AnimationAttributeNodes.h
+++ b/Engine/Source/Editor/AnimationBlueprintLibrary/Public/K2_AnimationAttributeNodes.h
@@ -38,6 +38,10 @@ public:
 	//~ End UK2Node Interface
 
 protected:
+	//~ Begin UK2Node_CallFunction interface
+	virtual bool CanToggleNodePurity() const override { return false; }
+	//~ End UK2Node_CallFunction interface
+
 	FName AttributeValuePinName = NAME_None;
 	FText AttributeActionFormat;
 	
diff --git a/Engine/Source/Editor/BlueprintGraph/Classes/K2Node_AddComponent.h b/Engine/Source/Editor/BlueprintGraph/Classes/K2Node_AddComponent.h
--- a/Engine/Source/Editor/BlueprintGraph/Classes/K2Node_AddComponent.h
+++ b/Engine/Source/Editor/BlueprintGraph/Classes/K2Node_AddComponent.h
@@ -100,6 +100,11 @@ class UK2Node_AddComponent : public UK2Node_CallFunction
 	/** Static name of function to call */
 	static BLUEPRINTGRAPH_API FName GetAddComponentFunctionName();
 
+protected:
+	//~ UK2Node_CallFunction interface
+	virtual bool CanToggleNodePurity() const override { return false; }
+	//~ End UK2Node_CallFunction interface
+
 private: 
 	UEdGraphPin* GetTemplateNamePin() const
 	{
diff --git a/Engine/Source/Editor/BlueprintGraph/Classes/K2Node_CallDataTableFunction.h b/Engine/Source/Editor/BlueprintGraph/Classes/K2Node_CallDataTableFunction.h
--- a/Engine/Source/Editor/BlueprintGraph/Classes/K2Node_CallDataTableFunction.h
+++ b/Engine/Source/Editor/BlueprintGraph/Classes/K2Node_CallDataTableFunction.h
@@ -21,4 +21,8 @@ class UK2Node_CallDataTableFunction : public UK2Node_CallFunction
 	virtual void PinDefaultValueChanged(UEdGraphPin* Pin) override;
 	virtual void NotifyPinConnectionListChanged(UEdGraphPin* Pin) override;
 	//~ End EdGraphNode Interface
+
+	//~ Begin UK2Node_CallFunction interface
+	virtual bool CanToggleNodePurity() const override { return false; }
+	//~ End UK2Node_CallFunction interface
 };
diff --git a/Engine/Source/Editor/BlueprintGraph/Classes/K2Node_CallFunction.h b/Engine/Source/Editor/BlueprintGraph/Classes/K2Node_CallFunction.h
--- a/Engine/Source/Editor/BlueprintGraph/Classes/K2Node_CallFunction.h
+++ b/Engine/Source/Editor/BlueprintGraph/Classes/K2Node_CallFunction.h
@@ -34,13 +34,24 @@ class UFunction;
 class UObject;
 template <typename KeyType, typename ValueType> struct TKeyValuePair;
 
+UENUM()
+enum class ENodePurityOverride : int8
+{
+	Unset = 0,
+	Pure,
+	Impure
+};
+
 UCLASS()
 class BLUEPRINTGRAPH_API UK2Node_CallFunction : public UK2Node
 {
 	GENERATED_UCLASS_BODY()
 
-	/** Indicates that this is a call to a pure function */
+	/** Indicates that the bound function defaults to a pure state */
 	UPROPERTY()
+	uint32 bDefaultsToPureFunc:1;
+	
+	UE_DEPRECATED(5.5, "bIsPureFunc is deprecated. Use bDefaultsToPureFunc or IsNodePure instead.")
 	uint32 bIsPureFunc:1;
 
 	/** Indicates that during compile we want to create multiple exec pins from an enum param */
@@ -118,7 +129,7 @@ public:
 
 	// UK2Node interface
 	virtual void ReallocatePinsDuringReconstruction(TArray<UEdGraphPin*>& OldPins) override;
-	virtual bool IsNodePure() const override { return bIsPureFunc; }
+	virtual bool IsNodePure() const override;
 	virtual void PostReconstructNode() override;
 	virtual bool ShouldDrawCompact() const override;
 	UE_DEPRECATED(5.4, "ShouldDrawAsBead is deprecated")
@@ -251,6 +262,13 @@ private:
 	/** Conforms container pins */
 	void ConformContainerPins();
 
+	UPROPERTY()
+	ENodePurityOverride NodePurityOverride;
+
+	bool AreExecPinsVisible() const;
+	bool FunctionHasOutputs() const;
+	void ToggleNodePurityOverride();
+
 protected:
 
 	/** Invalidates current pin tool tips, so that they will be refreshed before being displayed: */
@@ -259,6 +277,9 @@ protected:
 	/** Helper function to ensure function is called in our context */
 	virtual void FixupSelfMemberContext();
 
+	/** By default, pure nodes can be toggled. Return false if you don't want your node to support toggling. */
+	virtual bool CanToggleNodePurity() const;
+
 	/** Adds this function to the suppressed deprecation warnings list for this project */
 	void SuppressDeprecationWarning() const;
 
diff --git a/Engine/Source/Editor/BlueprintGraph/Classes/K2Node_CallFunctionOnMember.h b/Engine/Source/Editor/BlueprintGraph/Classes/K2Node_CallFunctionOnMember.h
--- a/Engine/Source/Editor/BlueprintGraph/Classes/K2Node_CallFunctionOnMember.h
+++ b/Engine/Source/Editor/BlueprintGraph/Classes/K2Node_CallFunctionOnMember.h
@@ -25,7 +25,7 @@ class UK2Node_CallFunctionOnMember : public UK2Node_CallFunction
 
 	/** Reference to member variable to call function on */
 	UPROPERTY()
-	FMemberReference				MemberVariableToCallOn;
+	FMemberReference MemberVariableToCallOn;
 
 	virtual bool HasExternalDependencies(TArray<class UStruct*>* OptionalOutput) const override;
 
@@ -34,7 +34,9 @@ class UK2Node_CallFunctionOnMember : public UK2Node_CallFunction
 	virtual FText GetFunctionContextString() const override;
 	virtual class FNodeHandlingFunctor* CreateNodeHandler(class FKismetCompilerContext& CompilerContext) const override;
 	virtual void ExpandNode(class FKismetCompilerContext& CompilerContext, UEdGraph* SourceGraph) override;
-	//~ End UK2Node_CallFunction Interface
 
+protected:
+	virtual bool CanToggleNodePurity() const override { return false; }
+	//~ End UK2Node_CallFunction Interface
 };
 
diff --git a/Engine/Source/Editor/BlueprintGraph/Classes/K2Node_InstancedStruct.h b/Engine/Source/Editor/BlueprintGraph/Classes/K2Node_InstancedStruct.h
--- a/Engine/Source/Editor/BlueprintGraph/Classes/K2Node_InstancedStruct.h
+++ b/Engine/Source/Editor/BlueprintGraph/Classes/K2Node_InstancedStruct.h
@@ -1,4 +1,4 @@
-﻿// Copyright Epic Games, Inc. All Rights Reserved.
+// Copyright Epic Games, Inc. All Rights Reserved.
 
 
 #pragma once
@@ -24,5 +24,10 @@ class UK2Node_InstancedStruct : public UK2Node_CallFunction
 	//~ Begin K2Node Interface
 	virtual bool IsConnectionDisallowed(const UEdGraphPin* MyPin, const UEdGraphPin* OtherPin, FString& OutReason) const override;
 	//~ End K2Node Interface
+
+protected:
+	//~ UK2Node_CallFunction interface
+	virtual bool CanToggleNodePurity() const override { return false; }
+	//~ End UK2Node_CallFunction interface
 };
 
diff --git a/Engine/Source/Editor/BlueprintGraph/Classes/K2Node_Message.h b/Engine/Source/Editor/BlueprintGraph/Classes/K2Node_Message.h
--- a/Engine/Source/Editor/BlueprintGraph/Classes/K2Node_Message.h
+++ b/Engine/Source/Editor/BlueprintGraph/Classes/K2Node_Message.h
@@ -25,24 +25,30 @@ class UK2Node_Message : public UK2Node_CallFunction
 	GENERATED_UCLASS_BODY()
 
 
-	//~ Begin UEdGraphNode Interface.
+	//~ Begin UEdGraphNode Interface
 	virtual FText GetNodeTitle(ENodeTitleType::Type TitleType) const override;
 	virtual FText GetTooltipText() const override;
 	virtual void AllocateDefaultPins() override;
-	//~ End UEdGraphNode Interface.
+	//~ End UEdGraphNode Interface
 
-	//~ Begin K2Node Interface.t
+	//~ Begin K2Node Interface
 	virtual bool IsNodeSafeToIgnore() const override { return true; }
 	virtual class FNodeHandlingFunctor* CreateNodeHandler(class FKismetCompilerContext& CompilerContext) const override;
 	virtual void ExpandNode(class FKismetCompilerContext& CompilerContext, UEdGraph* SourceGraph) override;
 	virtual FName GetCornerIcon() const override;
 	virtual bool IsNodePure() const override { return false; }
-	//~ End K2Node Interface.
+	//~ End K2Node Interface
 
 	//~ Begin K2Node_CallFunction Interface.
 	virtual UEdGraphPin* CreateSelfPin(const UFunction* Function) override;
 
 protected:
+	virtual void FixupSelfMemberContext() override;
+
+	virtual bool CanToggleNodePurity() const override { return false; }
+
+	//~ End K2Node_CallFunction Interface.
+
 	/**
 	 * Helper function for expand node to expand the message to handle the case that a ULevelStreaming may be passed which needs to have its LevelScriptActor pulled out
 	 *
@@ -54,10 +60,6 @@ protected:
 	 */
 	void ExpandLevelStreamingHandlers(class FKismetCompilerContext& InCompilerContext, UEdGraph* InSourceGraph, class UEdGraphPin* InStartingExecPin, class UEdGraphPin* InMessageSelfPin, class UK2Node_DynamicCast* InCastToInterfaceNode);
 
-	virtual void FixupSelfMemberContext() override;
-	//~ End K2Node_CallFunction Interface.
-
 	/** Constructing FText strings can be costly, so we cache the node's title */
 	FNodeTitleTextTable CachedNodeTitles;
 };
-
diff --git a/Engine/Source/Editor/BlueprintGraph/Classes/K2Node_PromotableOperator.h b/Engine/Source/Editor/BlueprintGraph/Classes/K2Node_PromotableOperator.h
--- a/Engine/Source/Editor/BlueprintGraph/Classes/K2Node_PromotableOperator.h
+++ b/Engine/Source/Editor/BlueprintGraph/Classes/K2Node_PromotableOperator.h
@@ -71,6 +71,11 @@ public:
 	/** Recombines all split pins and sets the node to have default values (all wildcard pins) */
 	BLUEPRINTGRAPH_API void ResetNodeToWildcard();
 	
+protected:
+	// UK2Node_CallFunction interface
+	virtual bool CanToggleNodePurity() const override { return false; }
+	// End of UK2Node_CallFunction interface
+
 private:
 
 	/**
diff --git a/Engine/Source/Editor/BlueprintGraph/Private/CallFunctionHandler.cpp b/Engine/Source/Editor/BlueprintGraph/Private/CallFunctionHandler.cpp
--- a/Engine/Source/Editor/BlueprintGraph/Private/CallFunctionHandler.cpp
+++ b/Engine/Source/Editor/BlueprintGraph/Private/CallFunctionHandler.cpp
@@ -26,6 +26,68 @@
 
 #define LOCTEXT_NAMESPACE "CallFunctionHandler"
 
+namespace UE::K2NodeCallFunction::Private
+{
+	static bool IsCalledFunctionPure(const UEdGraphNode* Node)
+	{
+		if (const UK2Node_CallFunction* CallFunctionNode = Cast<UK2Node_CallFunction>(Node))
+		{
+			return CallFunctionNode->IsNodePure();
+		}
+		
+		return false;
+	}
+
+	static bool FindMatchingReferencedNetOrFieldNotifyPropertyAndPin(TArray<UEdGraphPin*>& RemainingPins, FProperty* FunctionProperty, FProperty*& NetProperty, FProperty*& FieldNotifyProperty, UEdGraphPin*& PropertyObjectPin)
+	{
+		NetProperty = nullptr;
+		FieldNotifyProperty = nullptr;
+		PropertyObjectPin = nullptr;
+
+		if (UNLIKELY(FunctionProperty->HasAllPropertyFlags(CPF_OutParm | CPF_ReferenceParm) && !FunctionProperty->HasAnyPropertyFlags(CPF_ReturnParm | CPF_ConstParm)))
+		{
+			for (int32 i = 0; i < RemainingPins.Num(); ++i)
+			{
+				if (FunctionProperty->GetFName() == RemainingPins[i]->PinName)
+				{
+					bool bResult = false;
+					UEdGraphPin* ParamPin = RemainingPins[i];
+					RemainingPins.RemoveAtSwap(i);
+					if (UEdGraphPin* PinToTry = FEdGraphUtilities::GetNetFromPin(ParamPin))
+					{
+						// TODO: Should we traverse pin links to find references somehow?
+						//			E.G., How are select statements that pass through references
+						//					to net properties handled?
+
+						if (UK2Node_VariableGet* GetPropertyNode = Cast<UK2Node_VariableGet>(PinToTry->GetOwningNode()))
+						{
+							if (FProperty* ToCheck = GetPropertyNode->GetPropertyForVariable())
+							{
+								if (UNLIKELY(FKismetCompilerUtilities::IsPropertyUsesFieldNotificationSetValueAndBroadcast(ToCheck)))
+								{
+									FieldNotifyProperty = ToCheck;
+									PropertyObjectPin = GetPropertyNode->FindPinChecked(UEdGraphSchema_K2::PN_Self);
+									bResult = true;
+								}
+								if (UNLIKELY(ToCheck->HasAnyPropertyFlags(CPF_Net)))
+								{
+									NetProperty = ToCheck;
+									PropertyObjectPin = GetPropertyNode->FindPinChecked(UEdGraphSchema_K2::PN_Self);
+									bResult = true;
+								}
+							}
+						}
+					}
+
+					return bResult;
+				}
+			}
+		}
+
+		return false;
+	}
+}
+
 //////////////////////////////////////////////////////////////////////////
 // FImportTextErrorContext
 
@@ -545,7 +607,7 @@ void FKCHandler_CallFunction::CreateFunctionCallStatement(FKismetFunctionContext
 			else
 			{
 				// Generate the output impulse from this node
-				if (!IsCalledFunctionPure(Node))
+				if (!UE::K2NodeCallFunction::Private::IsCalledFunctionPure(Node))
 				{
 					GenerateSimpleThenGoto(Context, *Node);
 				}
@@ -720,58 +782,6 @@ UFunction* FKCHandler_CallFunction::FindFunction(FKismetFunctionContext& Context
 	return nullptr;
 }
 
-namespace UE::BlueprintGraph::Private
-{
-static bool FindMatchingReferencedNetOrFieldNotifyPropertyAndPin(TArray<UEdGraphPin*>& RemainingPins, FProperty* FunctionProperty, FProperty*& NetProperty, FProperty*& FieldNotifyProperty, UEdGraphPin*& PropertyObjectPin)
-{
-	NetProperty = nullptr;
-	FieldNotifyProperty = nullptr;
-	PropertyObjectPin = nullptr;
-
-	if (UNLIKELY(FunctionProperty->HasAllPropertyFlags(CPF_OutParm | CPF_ReferenceParm) && !FunctionProperty->HasAnyPropertyFlags(CPF_ReturnParm | CPF_ConstParm)))
-	{
-		for (int32 i = 0; i < RemainingPins.Num(); ++i)
-		{
-			if (FunctionProperty->GetFName() == RemainingPins[i]->PinName)
-			{
-				bool bResult = false;
-				UEdGraphPin* ParamPin = RemainingPins[i];
-				RemainingPins.RemoveAtSwap(i);
-				if (UEdGraphPin* PinToTry = FEdGraphUtilities::GetNetFromPin(ParamPin))
-				{
-					// TODO: Should we traverse pin links to find references somehow?
-					//			E.G., How are select statements that pass through references
-					//					to net properties handled?
-
-					if (UK2Node_VariableGet* GetPropertyNode = Cast<UK2Node_VariableGet>(PinToTry->GetOwningNode()))
-					{
-						if (FProperty* ToCheck = GetPropertyNode->GetPropertyForVariable())
-						{
-							if (UNLIKELY(FKismetCompilerUtilities::IsPropertyUsesFieldNotificationSetValueAndBroadcast(ToCheck)))
-							{
-								FieldNotifyProperty = ToCheck;
-								PropertyObjectPin = GetPropertyNode->FindPinChecked(UEdGraphSchema_K2::PN_Self);
-								bResult = true;
-							}
-							if (UNLIKELY(ToCheck->HasAnyPropertyFlags(CPF_Net)))
-							{
-								NetProperty = ToCheck;
-								PropertyObjectPin = GetPropertyNode->FindPinChecked(UEdGraphSchema_K2::PN_Self);
-								bResult = true;
-							}
-						}
-					}
-				}
-
-				return bResult;
-			}
-		}
-	}
-
-	return false;
-}
-}
-
 void FKCHandler_CallFunction::Transform(FKismetFunctionContext& Context, UEdGraphNode* Node)
 {
 	// Add an object reference pin for this call
@@ -781,7 +791,7 @@ void FKCHandler_CallFunction::Transform(FKismetFunctionContext& Context, UEdGrap
 		return;
 	}
 
-	const bool bIsPure = CallFuncNode->bIsPureFunc;
+	const bool bIsPure = CallFuncNode->IsNodePure();
 	bool bIsPureAndNoUsedOutputs = false;
 	if (bIsPure)
 	{
@@ -890,7 +900,7 @@ void FKCHandler_CallFunction::Transform(FKismetFunctionContext& Context, UEdGrap
 				// This is similar to the loop in CreateCallFunction
 				for (TFieldIterator<FProperty> It(Function); It; ++It)
 				{
-					if (UE::BlueprintGraph::Private::FindMatchingReferencedNetOrFieldNotifyPropertyAndPin(RemainingPins, *It, NetProperty, FieldNotifyProperty, PropertyObjectPin))
+					if (UE::K2NodeCallFunction::Private::FindMatchingReferencedNetOrFieldNotifyPropertyAndPin(RemainingPins, *It, NetProperty, FieldNotifyProperty, PropertyObjectPin))
 					{
 						if (bIsPure)
 						{
@@ -976,7 +986,7 @@ void FKCHandler_CallFunction::Compile(FKismetFunctionContext& Context, UEdGraphN
 	check(NULL != Node);
 
 	//@TODO: Can probably move this earlier during graph verification instead of compilation, but after island pruning
-	if (!IsCalledFunctionPure(Node))
+	if (!UE::K2NodeCallFunction::Private::IsCalledFunctionPure(Node))
 	{
 		// For imperative nodes, make sure the exec function was actually triggered and not just included due to an output data dependency
 		UEdGraphPin* ExecTriggeringPin = CompilerContext.GetSchema()->FindExecutionPin(*Node, EGPD_Input);
diff --git a/Engine/Source/Editor/BlueprintGraph/Private/K2Node_AddComponent.cpp b/Engine/Source/Editor/BlueprintGraph/Private/K2Node_AddComponent.cpp
--- a/Engine/Source/Editor/BlueprintGraph/Private/K2Node_AddComponent.cpp
+++ b/Engine/Source/Editor/BlueprintGraph/Private/K2Node_AddComponent.cpp
@@ -52,7 +52,7 @@ const FName UK2Node_AddComponent::NAME_ManualAttachment(TEXT("bManualAttachment"
 UK2Node_AddComponent::UK2Node_AddComponent(const FObjectInitializer& ObjectInitializer)
 	: Super(ObjectInitializer)
 {
-	bIsPureFunc = false;
+	bDefaultsToPureFunc = false;
 }
 
 // We add this prefix to template object names.
diff --git a/Engine/Source/Editor/BlueprintGraph/Private/K2Node_CallFunction.cpp b/Engine/Source/Editor/BlueprintGraph/Private/K2Node_CallFunction.cpp
--- a/Engine/Source/Editor/BlueprintGraph/Private/K2Node_CallFunction.cpp
+++ b/Engine/Source/Editor/BlueprintGraph/Private/K2Node_CallFunction.cpp
@@ -857,6 +857,11 @@ void UK2Node_CallFunction::ReallocatePinsDuringReconstruction(TArray<UEdGraphPin
 	ReconnectPureExecPins(OldPins);
 }
 
+bool UK2Node_CallFunction::IsNodePure() const
+{
+	return !AreExecPinsVisible();
+}
+
 UEdGraphPin* UK2Node_CallFunction::CreateSelfPin(const UFunction* Function)
 {
 	return FBlueprintNodeStatics::CreateSelfPin(this, Function);
@@ -870,7 +875,7 @@ void UK2Node_CallFunction::CreateExecPinsForFunctionCall(const UFunction* Functi
 	ExpandAsEnumPins.Reset();
 
 	// If not pure, create exec pins
-	if (!bIsPureFunc)
+	if (!IsNodePure())
 	{
 		// If we want enum->exec expansion, and it is not disabled, do it now
 		if (bWantsEnumToExecExpansion)
@@ -1127,7 +1132,7 @@ bool UK2Node_CallFunction::CreatePinsForFunctionCall(const UFunction* Function)
 
 	UClass* FunctionOwnerClass = Function->GetOuterUClass();
 
-	bIsPureFunc = (Function->HasAnyFunctionFlags(FUNC_BlueprintPure) != false);
+	bDefaultsToPureFunc = Function->HasAnyFunctionFlags(FUNC_BlueprintPure);
 	DetermineWantsEnumToExecExpansion(Function);
 
 	// Create input pins
@@ -1173,7 +1178,7 @@ bool UK2Node_CallFunction::CreatePinsForFunctionCall(const UFunction* Function)
 
 				// Non-const pure functions are generally compact looking, so we take advantage of that by hiding the self pin if we can.
 				// In this case, if the function belongs to our current class, then "self" is implied.
-				SelfPin->bHidden = (bIsFunctionCompatibleWithSelf && bIsPureFunc && !bLocalIsConstFunc);
+				SelfPin->bHidden = (bIsFunctionCompatibleWithSelf && IsNodePure() && !bLocalIsConstFunc);
 			}
 		}
 	}
@@ -1321,6 +1326,18 @@ void UK2Node_CallFunction::PostReconstructNode()
 		// Remove the breakpoint
 		FKismetDebugUtilities::RemoveBreakpointFromNode(this, GetBlueprint());
 	}
+
+	// If the user marked the override as pure, but the underlying function changed to have no outputs,
+	// then we need to reset the override. Nodes can only be pure if they have outputs.
+	const bool bNeedsPurityOverrideReset =
+		(NodePurityOverride == ENodePurityOverride::Pure) &&
+		!FunctionHasOutputs()
+	;
+	
+	if (bNeedsPurityOverrideReset)
+	{
+		NodePurityOverride = ENodePurityOverride::Unset;
+	}
 }
 
 void UK2Node_CallFunction::NotifyPinConnectionListChanged(UEdGraphPin* Pin)
@@ -1879,7 +1896,7 @@ void UK2Node_CallFunction::SetFromFunction(const UFunction* Function)
 {
 	if (Function != NULL)
 	{
-		bIsPureFunc = Function->HasAnyFunctionFlags(FUNC_BlueprintPure);
+		bDefaultsToPureFunc = Function->HasAnyFunctionFlags(FUNC_BlueprintPure);
 		DetermineWantsEnumToExecExpansion(Function);
 
 		FunctionReference.SetFromField<UFunction>(Function, GetBlueprintClassFromNode());
@@ -2066,6 +2083,13 @@ void UK2Node_CallFunction::FixupSelfMemberContext()
 	}
 }
 
+bool UK2Node_CallFunction::CanToggleNodePurity() const
+{
+	// Only functions with outputs can be toggled.
+	// Otherwise, you can end up with a pure node with no outputs that's never evaluated.
+	return FunctionHasOutputs();
+}
+
 void UK2Node_CallFunction::SuppressDeprecationWarning() const
 {
 	if (UFunction* Function = GetTargetFunction())
@@ -3255,7 +3279,7 @@ void UK2Node_CallFunction::ConformContainerPins()
 	const FString& MapKeyPinMetaData = TargetFunction->GetMetaData(FBlueprintMetadata::MD_MapKeyParam);
 	const FString& MapValuePinMetaData = TargetFunction->GetMetaData(FBlueprintMetadata::MD_MapValueParam);
 
-	if(!MapPinMetaData.IsEmpty() || !MapKeyPinMetaData.IsEmpty() || !MapValuePinMetaData.IsEmpty() )
+	if (!MapPinMetaData.IsEmpty() || !MapKeyPinMetaData.IsEmpty() || !MapValuePinMetaData.IsEmpty())
 	{
 		// if the map pin has a connection infer from that, otherwise use the information on the key param and value param:
 		bool bReadyToPropagateKeyType = false;
@@ -3280,6 +3304,68 @@ void UK2Node_CallFunction::ConformContainerPins()
 	}
 }
 
+bool UK2Node_CallFunction::AreExecPinsVisible() const
+{
+	const bool bAreExecPinsVisible =
+		(NodePurityOverride == ENodePurityOverride::Impure) ||
+		(!bDefaultsToPureFunc && NodePurityOverride != ENodePurityOverride::Pure)
+	;
+
+	return bAreExecPinsVisible;
+}
+
+bool UK2Node_CallFunction::FunctionHasOutputs() const
+{
+	bool bCanToggle = false;
+
+	const UFunction* Function = GetTargetFunction();
+	if (ensure(Function))
+	{
+		for (TFieldIterator<FProperty> PropIt(Function); PropIt && (PropIt->PropertyFlags & CPF_Parm); ++PropIt)
+		{
+			bCanToggle =
+				PropIt->HasAnyPropertyFlags(CPF_ReturnParm) ||
+				PropIt->HasAnyPropertyFlags(CPF_OutParm)
+				;
+
+			if (bCanToggle)
+			{
+				break;
+			}
+		}
+	}
+
+	return bCanToggle;
+}
+
+void UK2Node_CallFunction::ToggleNodePurityOverride()
+{
+	const FText TransactionTitle = AreExecPinsVisible() ? LOCTEXT("HideExecPins", "Hide Exec pins") : LOCTEXT("ShowExecPins", "Show Exec pins");
+	
+	FScopedTransaction Transaction(TransactionTitle);
+	Modify();
+
+	switch (NodePurityOverride)
+	{
+	case ENodePurityOverride::Unset:
+		NodePurityOverride = bDefaultsToPureFunc ? ENodePurityOverride::Impure : ENodePurityOverride::Pure;
+		break;
+
+	case ENodePurityOverride::Pure:
+		NodePurityOverride = ENodePurityOverride::Impure;
+		break;
+
+	case ENodePurityOverride::Impure:
+		NodePurityOverride = ENodePurityOverride::Pure;
+		break;
+
+	default:
+		checkNoEntry();
+	}
+	
+	ReconstructNode();
+}
+
 FText UK2Node_CallFunction::GetToolTipHeading() const
 {
 	FText Heading = Super::GetToolTipHeading();
@@ -3617,12 +3703,13 @@ void UK2Node_CallFunction::GetNodeContextMenuActions(class UToolMenu* Menu, clas
 {
 	Super::GetNodeContextMenuActions(Menu, Context);
 
+	FToolMenuSection& Section = Menu->AddSection("K2NodeCallFunction", LOCTEXT("FunctionHeader", "Function"));
+
 	if (HasDeprecatedReference())
 	{
-		FText MenuEntryTitle = LOCTEXT("SuppressFunctionDeprecationWarningTitle", "Suppress Deprecation Warning");
-		FText MenuEntryTooltip = LOCTEXT("SuppressFunctionDeprecationWarningTooltip", "Adds this function to the suppressed deprecation warnings list in the Bluperint Editor Project Settings for this project.");
+		const FText MenuEntryTitle = LOCTEXT("SuppressFunctionDeprecationWarningTitle", "Suppress Deprecation Warning");
+		const FText MenuEntryTooltip = LOCTEXT("SuppressFunctionDeprecationWarningTooltip", "Adds this function to the suppressed deprecation warnings list in the Bluperint Editor Project Settings for this project.");
 
-		FToolMenuSection& Section = Menu->AddSection("K2NodeCallFunction", LOCTEXT("FunctionHeader", "Function"));
 		Section.AddMenuEntry(
 			"SuppressDeprecationWarning",
 			MenuEntryTitle,
@@ -3635,6 +3722,30 @@ void UK2Node_CallFunction::GetNodeContextMenuActions(class UToolMenu* Menu, clas
 			)
 		);
 	}
+
+	if (CanToggleNodePurity())
+	{
+		const FText HideExecPinsTitle = LOCTEXT("HideExecPins", "Hide Exec pins");
+		const FText ShowExecPinsTitle = LOCTEXT("ShowExecPins", "Show Exec pins");
+		const FText HideExecPinsTooltip = LOCTEXT("HideExecPinsTooltip", "Hide the node's Exec pins. The node will be Pure and will execute once for each node connected to an output pin.");
+		const FText ShowExecPinsTooltip = LOCTEXT("ShowExecPinsTooltip", "Show the node's Exec pins. The node will not be Pure and will have execution pins that must be connected to the graph. The outputs may also be cached for later use.");
+
+		const bool bAreExecPinsVisible = AreExecPinsVisible();
+		const FText& MenuEntryTitle = bAreExecPinsVisible ? HideExecPinsTitle : ShowExecPinsTitle;
+		const FText& MenuEntryTooltip = bAreExecPinsVisible ? HideExecPinsTooltip : ShowExecPinsTooltip;
+
+		Section.AddMenuEntry(
+			TEXT("ToggleExecPins"),
+			MenuEntryTitle,
+			MenuEntryTooltip,
+			FSlateIcon(),
+			FUIAction(
+				FExecuteAction::CreateUObject(const_cast<UK2Node_CallFunction*>(this), &UK2Node_CallFunction::ToggleNodePurityOverride),
+				FCanExecuteAction(),
+				FIsActionChecked()
+			)
+		);
+	}
 }
 
 FString UK2Node_CallFunction::GetPinMetaData(FName InPinName, FName InKey)
diff --git a/Engine/Source/Editor/BlueprintGraph/Private/K2Node_CallParentFunction.cpp b/Engine/Source/Editor/BlueprintGraph/Private/K2Node_CallParentFunction.cpp
--- a/Engine/Source/Editor/BlueprintGraph/Private/K2Node_CallParentFunction.cpp
+++ b/Engine/Source/Editor/BlueprintGraph/Private/K2Node_CallParentFunction.cpp
@@ -66,9 +66,9 @@ void UK2Node_CallParentFunction::AllocateDefaultPins()
 
 void UK2Node_CallParentFunction::SetFromFunction(const UFunction* Function)
 {
-	if (Function != NULL)
+	if (Function != nullptr)
 	{
-		bIsPureFunc = Function->HasAnyFunctionFlags(FUNC_BlueprintPure);
+		bDefaultsToPureFunc = Function->HasAnyFunctionFlags(FUNC_BlueprintPure);
 
 		UClass* OwnerClass = Function->GetOwnerClass();
 
@@ -79,7 +79,8 @@ void UK2Node_CallParentFunction::SetFromFunction(const UFunction* Function)
 			UBlueprint::GetGuidFromClassByFieldName<UFunction>(OwnerClass, Function->GetFName(), FunctionGuid);
 		}
 
-		FunctionReference.SetDirect(Function->GetFName(), FunctionGuid, OwnerClass, /*bIsConsideredSelfContext =*/false);
+		const bool bIsConsideredSelfContext = false;
+		FunctionReference.SetDirect(Function->GetFName(), FunctionGuid, OwnerClass, bIsConsideredSelfContext);
 	}
 }
 
diff --git a/Engine/Source/Editor/BlueprintGraph/Private/K2Node_EditorPropertyAccess.h b/Engine/Source/Editor/BlueprintGraph/Private/K2Node_EditorPropertyAccess.h
--- a/Engine/Source/Editor/BlueprintGraph/Private/K2Node_EditorPropertyAccess.h
+++ b/Engine/Source/Editor/BlueprintGraph/Private/K2Node_EditorPropertyAccess.h
@@ -30,6 +30,11 @@ public:
 	virtual bool IsActionFilteredOut(const class FBlueprintActionFilter& Filter) override;
 	virtual void EarlyValidation(class FCompilerResultsLog& MessageLog) const override;
 	//~ End UK2Node Interface
+
+protected:
+	//~ Begin UK2Node_CallFunction interface
+	virtual bool CanToggleNodePurity() const override { return false; }
+	//~ End UK2Node_CallFunction interface
 };
 
 UCLASS()
diff --git a/Engine/Source/Editor/BlueprintGraph/Public/CallFunctionHandler.h b/Engine/Source/Editor/BlueprintGraph/Public/CallFunctionHandler.h
--- a/Engine/Source/Editor/BlueprintGraph/Public/CallFunctionHandler.h
+++ b/Engine/Source/Editor/BlueprintGraph/Public/CallFunctionHandler.h
@@ -35,11 +35,12 @@ public:
 	 */
 	void CreateFunctionCallStatement(FKismetFunctionContext& Context, UEdGraphNode* Node, UEdGraphPin* SelfPin);
 
+	UE_DEPRECATED(5.5, "IsCalledFunctionPure is deprecated")
 	bool IsCalledFunctionPure(UEdGraphNode* Node)
 	{
 		if (UK2Node_CallFunction* CallFunctionNode = Cast<UK2Node_CallFunction>(Node))
 		{
-			return CallFunctionNode->bIsPureFunc;
+			return CallFunctionNode->IsNodePure();
 		}
 
 		return false;
