Added project setting to control toggling of impure Blueprint functions to pure ones (disabled by default).

(Follow up change to CL 35309066)

If enabled, a function that isn't marked as pure can be converted to a pure node in a graph. While this provides flexibility for controlling graph layouts, pure nodes can introduce hidden performance costs. If an impure function is expensive, its conversion to a pure can introduce multiple calls for each connected output. Since Blueprint graphs don't do a great job of visualizing executions of pure nodes, we're defaulting this feature to off. Teams are free to decide if they wish to use it.

Additionally, this will issue a compiler warning if a illegal pure override is discovered in a node. This can occur if a user locally enables the project setting and checks in converted nodes.

diff --git a/Engine/Source/Editor/BlueprintGraph/Private/K2Node_CallFunction.cpp b/Engine/Source/Editor/BlueprintGraph/Private/K2Node_CallFunction.cpp
--- a/Engine/Source/Editor/BlueprintGraph/Private/K2Node_CallFunction.cpp
+++ b/Engine/Source/Editor/BlueprintGraph/Private/K2Node_CallFunction.cpp
@@ -2085,9 +2085,24 @@ void UK2Node_CallFunction::FixupSelfMemberContext()
 
 bool UK2Node_CallFunction::CanToggleNodePurity() const
 {
-	// Only functions with outputs can be toggled.
+	// If the underlying function defaults to pure, then we can always toggle it.
+	// If not, then we only allow it if enabled in the project settings.
+	// Additionally, only functions with ouputs can be toggled.
 	// Otherwise, you can end up with a pure node with no outputs that's never evaluated.
-	return FunctionHasOutputs();
+
+	const bool bFunctionHasOutputs = FunctionHasOutputs();
+
+	if (bDefaultsToPureFunc)
+	{
+		return bFunctionHasOutputs;
+	}
+	else
+	{
+		const UBlueprintEditorProjectSettings* EditorProjectSettings = GetDefault<UBlueprintEditorProjectSettings>();
+		check(EditorProjectSettings);
+
+		return bFunctionHasOutputs && EditorProjectSettings->bAllowImpureToPureNodeConversion;
+	}
 }
 
 void UK2Node_CallFunction::SuppressDeprecationWarning() const
@@ -2204,7 +2219,7 @@ void UK2Node_CallFunction::ValidateNodeDuringCompilation(class FCompilerResultsL
 
 	const UBlueprint* Blueprint = GetBlueprint();
 	UFunction *Function = GetTargetFunction();
-	if (Function == NULL)
+	if (Function == nullptr)
 	{
 		FString OwnerName;
 
@@ -2291,7 +2306,7 @@ void UK2Node_CallFunction::ValidateNodeDuringCompilation(class FCompilerResultsL
 				}
 			}
 
-			if ( bNodeIsInConstructionScript )
+			if (bNodeIsInConstructionScript)
 			{
 				MessageLog.Warning(*LOCTEXT("FunctionUnsafeDuringConstruction", "Function '@@' is unsafe to call in a construction script.").ToString(), this);
 			}
@@ -2323,7 +2338,7 @@ void UK2Node_CallFunction::ValidateNodeDuringCompilation(class FCompilerResultsL
 			}
 		}
 
-		if(Blueprint && !FBlueprintEditorUtils::IsNativeSignature(Function))
+		if (Blueprint && !FBlueprintEditorUtils::IsNativeSignature(Function))
 		{
 			// enforce protected function restriction
 			const bool bCanTreatAsError = Blueprint->GetLinkerCustomVersion(FFrameworkObjectVersion::GUID) >= FFrameworkObjectVersion::EnforceBlueprintFunctionVisibility;
@@ -2347,7 +2362,7 @@ void UK2Node_CallFunction::ValidateNodeDuringCompilation(class FCompilerResultsL
 			const bool bFuncBelongsToClass = bFuncBelongsToSubClass && (Blueprint->SkeletonGeneratedClass == Function->GetOuterUClass());
 			if (bIsPrivate && !bFuncBelongsToClass)
 			{
-				if(bCanTreatAsError)
+				if (bCanTreatAsError)
 				{
 					MessageLog.Error(*LOCTEXT("FunctionPrivateAccessed", "Function '@@' is private and can't be accessed outside of its defined class '@@'.").ToString(), this, Function->GetOuterUClass());
 				}
@@ -2357,6 +2372,20 @@ void UK2Node_CallFunction::ValidateNodeDuringCompilation(class FCompilerResultsL
 				}
 			}
 		}
+
+		const UBlueprintEditorProjectSettings* EditorProjectSettings = GetDefault<UBlueprintEditorProjectSettings>();
+		check(EditorProjectSettings);
+
+		const bool bHasIllegalPureOverride =
+			!bDefaultsToPureFunc &&
+			(NodePurityOverride == ENodePurityOverride::Pure) &&
+			(!EditorProjectSettings->bAllowImpureToPureNodeConversion)
+		;
+
+		if (bHasIllegalPureOverride)
+		{
+			MessageLog.Warning(*LOCTEXT("IllegalPureNodeOverride", "Function '@@' is not pure and conversions to pure nodes are illegal in this project. Either replace the node or enable pure node conversions in the project settings.").ToString(), this);
+		}
 	}
 
 	FDynamicOutputHelper::VerifyNode(this, MessageLog);
diff --git a/Engine/Source/Editor/UnrealEd/Public/Settings/BlueprintEditorProjectSettings.h b/Engine/Source/Editor/UnrealEd/Public/Settings/BlueprintEditorProjectSettings.h
--- a/Engine/Source/Editor/UnrealEd/Public/Settings/BlueprintEditorProjectSettings.h
+++ b/Engine/Source/Editor/UnrealEd/Public/Settings/BlueprintEditorProjectSettings.h
@@ -37,6 +37,15 @@ public:
 	UPROPERTY(EditAnywhere, config, Category = Experimental)
 	uint8 bEnableChildActorExpansionInTreeView : 1;
 
+	/**
+	 * In Blueprint graphs, allows conversion of impure nodes to pure ones.
+	 * While this allows full flexibility of graph node visualization, this could potentially introduce unwanted performance hits.
+	 * Specifically, pure nodes are evaluated for each connected output. If an expensive function node were converted to pure, the user
+	 * might not realize that there's now added overhead.
+	 */
+	UPROPERTY(EditAnywhere, config, Category = Blueprints, DisplayName = "Allow conversion of impure nodes to pure ones")
+	uint8 bAllowImpureToPureNodeConversion : 1;
+
 	/**
 	 * Default view mode to use for child actor components in a Blueprint actor's component tree hierarchy (experimental).
 	 */
