Re-enable smart macro inference (which handles more than one connected type of pin), but be less eager to propagate wildcards when in a macro graph

#rnx

diff --git a/Engine/Source/Editor/BlueprintGraph/Private/K2Node_MacroInstance.cpp b/Engine/Source/Editor/BlueprintGraph/Private/K2Node_MacroInstance.cpp
--- a/Engine/Source/Editor/BlueprintGraph/Private/K2Node_MacroInstance.cpp
+++ b/Engine/Source/Editor/BlueprintGraph/Private/K2Node_MacroInstance.cpp
@@ -295,7 +295,13 @@ void UK2Node_MacroInstance::NotifyPinConnectionListChanged(UEdGraphPin* ChangedP
 			{
 				// we found one, infer from it and then propagate the inference:
 				FWildcardNodeUtils::InferType(ChangedPin, InferrablePin->PinType);
-				InferWildcards();
+
+				const UEdGraph* Graph = GetGraph();
+				const bool bIsMacroGraph = (Graph->GetSchema()->GetGraphType(Graph) == GT_Macro);
+				if (!bIsMacroGraph)
+				{
+					InferWildcards();
+				}
 			}
 		}
 	}
@@ -396,8 +402,16 @@ void UK2Node_MacroInstance::PostReconstructNode()
 			}
 		}
 
-		// rerun inference
-		InferWildcards();
+		
+		const UEdGraph* Graph = GetGraph();
+		const bool bIsMacroGraph = (Graph->GetSchema()->GetGraphType(Graph) == GT_Macro);
+		UBlueprint* Blueprint  = GetBlueprint();
+		const bool bIsCompiling = Blueprint ? Blueprint->bBeingCompiled : false;
+		if(!bIsMacroGraph || !bIsCompiling)
+		{
+			// rerun inference
+			InferWildcards();
+		}
 	}
 	else
 	{
@@ -651,7 +665,7 @@ void UK2Node_MacroInstance::InferWildcards()
 		// perform macro expansion in a dummy graph, inferring whatever types we can from the provided wildcards:
 		FCompilerResultsLog MessageLog;
 		UBlueprint* BP = GetBlueprint();
-		UEdGraph* ClonedGraph = FEdGraphUtilities::CloneGraph(MacroGraph, GetBlueprint(), &MessageLog, true);
+		UEdGraph* ClonedGraph = FEdGraphUtilities::CloneGraph(MacroGraph, BP, &MessageLog, true);
 		if (ClonedGraph)
 		{
 			InferWildcards(ClonedGraph->Nodes);
diff --git a/Engine/Source/Editor/BlueprintGraph/Private/K2Node_Tunnel.cpp b/Engine/Source/Editor/BlueprintGraph/Private/K2Node_Tunnel.cpp
--- a/Engine/Source/Editor/BlueprintGraph/Private/K2Node_Tunnel.cpp
+++ b/Engine/Source/Editor/BlueprintGraph/Private/K2Node_Tunnel.cpp
@@ -280,15 +280,8 @@ bool UK2Node_Tunnel::ModifyUserDefinedPinDefaultValue(TSharedPtr<FUserPinInfo> P
 
 bool UK2Node_Tunnel::ShouldDoSmartWildcardInference()
 {
-#if 0
 	static const FBoolConfigValueHelper bUseSimpleWildcardInference(TEXT("Blueprints"), TEXT("bUseSimpleWildcardInference"), GEngineIni);
 	return !bUseSimpleWildcardInference;
-#else
-	// This cannot be enabled until some leniency is added - e.g. to allow
-	// TMaps with a more restrictive value type to supercede inferences of base 
-	// type functionality
-	return false;
-#endif
 }
 
 bool UK2Node_Tunnel::CanModifyExecutionWires()
