Fixed double exec pins with message nodes

Prior to CL 35309066, CreateExecPinsForFunctionCall would only create exec pins if the underlying function had FUNC_BlueprintPure set. However, since message nodes are always impure, the AllocateDefaultPins override would force exec pins to show even if the function is pure.

CL 35309066 changed the logic in CreateExecPinsForFunctionCall to use IsNodePure for determining if exec pins should display. Since the message node overrides IsNodePure to false, we could get a scenario where both the overridden AllocateDefaultPins and CreateExecPinsForFunctionCall would create exec pins. The most straightforward fix is to remove the AllocateDefaultPins override, and let CreateExecPinsForFunctionCall handle exec pin creation.

Also resaved affected content that needed nodes refreshed.

#rnx

#changelist validated
#virtualized

diff --git a/Engine/Source/Editor/BlueprintGraph/Classes/K2Node_Message.h b/Engine/Source/Editor/BlueprintGraph/Classes/K2Node_Message.h
--- a/Engine/Source/Editor/BlueprintGraph/Classes/K2Node_Message.h
+++ b/Engine/Source/Editor/BlueprintGraph/Classes/K2Node_Message.h
@@ -28,7 +28,6 @@ class UK2Node_Message : public UK2Node_CallFunction
 	//~ Begin UEdGraphNode Interface
 	virtual FText GetNodeTitle(ENodeTitleType::Type TitleType) const override;
 	virtual FText GetTooltipText() const override;
-	virtual void AllocateDefaultPins() override;
 	//~ End UEdGraphNode Interface
 
 	//~ Begin K2Node Interface
diff --git a/Engine/Source/Editor/BlueprintGraph/Private/K2Node_Message.cpp b/Engine/Source/Editor/BlueprintGraph/Private/K2Node_Message.cpp
--- a/Engine/Source/Editor/BlueprintGraph/Private/K2Node_Message.cpp
+++ b/Engine/Source/Editor/BlueprintGraph/Private/K2Node_Message.cpp
@@ -83,25 +83,6 @@ FText UK2Node_Message::GetTooltipText() const
 	return CachedTooltip;
 }
 
-void UK2Node_Message::AllocateDefaultPins()
-{
-	UFunction* MessageNodeFunction = GetTargetFunction();
-	// since we have branching logic in ExpandNode(), this has to be an impure
-	// node with exec pins
-	//
-	// @TODO: make it so we can have impure message nodes using a custom 
-	//        FNodeHandlingFunctor, instead of ExpandNode()
-	if (MessageNodeFunction && MessageNodeFunction->HasAnyFunctionFlags(FUNC_BlueprintPure))
-	{
-		// Input - Execution Pin
-		CreatePin(EGPD_Input,  UEdGraphSchema_K2::PC_Exec, UEdGraphSchema_K2::PN_Execute);
-		// Output - Execution Pin
-		CreatePin(EGPD_Output, UEdGraphSchema_K2::PC_Exec, UEdGraphSchema_K2::PN_Then);
-	}
-
-	Super::AllocateDefaultPins();
-}
-
 UEdGraphPin* UK2Node_Message::CreateSelfPin(const UFunction* Function)
 {
 	UEdGraphPin* SelfPin = CreatePin(EGPD_Input, UEdGraphSchema_K2::PC_Object, UObject::StaticClass(), UEdGraphSchema_K2::PN_Self);
