Ensure that function parameter names can't be renamed to any local variables.

We already have a check that prevents renaming local variables to parameters, but we don't have the opposite check. Since renames first test the pin rename with RenameUserDefinedPin, we need an implementation in the function entry node that ensures that the new name doesn't collide with the local variable names.

diff --git a/Engine/Source/Editor/BlueprintGraph/Classes/K2Node_FunctionEntry.h b/Engine/Source/Editor/BlueprintGraph/Classes/K2Node_FunctionEntry.h
--- a/Engine/Source/Editor/BlueprintGraph/Classes/K2Node_FunctionEntry.h
+++ b/Engine/Source/Editor/BlueprintGraph/Classes/K2Node_FunctionEntry.h
@@ -87,6 +87,7 @@ class UK2Node_FunctionEntry : public UK2Node_FunctionTerminator
 	virtual void PostReconstructNode() override;
 	virtual void ClearCachedBlueprintData(UBlueprint* Blueprint) override;
 	virtual void FixupPinStringDataReferences(FArchive* SavingArchive) override;
+	virtual ERenamePinResult RenameUserDefinedPinImpl(FName OldName, FName NewName, bool bTest) override;
 	//~ End UK2Node Interface
 
 	//~ Begin UK2Node_EditablePinBase Interface
diff --git a/Engine/Source/Editor/BlueprintGraph/Private/K2Node_FunctionEntry.cpp b/Engine/Source/Editor/BlueprintGraph/Private/K2Node_FunctionEntry.cpp
--- a/Engine/Source/Editor/BlueprintGraph/Private/K2Node_FunctionEntry.cpp
+++ b/Engine/Source/Editor/BlueprintGraph/Private/K2Node_FunctionEntry.cpp
@@ -405,7 +405,7 @@ void UK2Node_FunctionEntry::AllocateDefaultPins()
 
 	if (FFunctionEntryHelper::RequireWorldContextParameter(this) 
 		&& ensureMsgf(!FindPin(FFunctionEntryHelper::GetWorldContextPinName()), 
-		TEXT("%s: World context parameter pin already exiss on function entry node %s"), 
+		TEXT("%s: World context parameter pin already exists on function entry node %s"), 
 			*GetOutermost()->GetName(),
 			*(CustomGeneratedFunctionName.IsNone() ? FunctionReference.GetMemberName() : CustomGeneratedFunctionName).ToString()
 		))
@@ -1078,6 +1078,33 @@ void UK2Node_FunctionEntry::FixupPinStringDataReferences(FArchive* SavingArchive
 	}
 }
 
+ERenamePinResult UK2Node_FunctionEntry::RenameUserDefinedPinImpl(FName OldName, FName NewName, bool bTest)
+{
+	ERenamePinResult Result = Super::RenameUserDefinedPinImpl(OldName, NewName, bTest);
+
+	// We don't actually rename our local variables via this method,
+	// but we do need to verify that no name collisions will occur when testing a pin rename.
+
+	const bool bCheckLocalVariablesNameCollision =
+		(Result == ERenamePinResult::ERenamePinResult_Success) &&
+		bTest
+	;
+
+	if (bCheckLocalVariablesNameCollision)
+	{
+		for (const FBPVariableDescription& Entry : LocalVariables)
+		{
+			if (NewName == Entry.VarName)
+			{
+				Result = ERenamePinResult_NameCollision;
+				break;
+			}
+		}
+	}
+
+	return Result;
+}
+
 bool UK2Node_FunctionEntry::ModifyUserDefinedPinDefaultValue(TSharedPtr<FUserPinInfo> PinInfo, const FString& NewDefaultValue)
 {
 	if (Super::ModifyUserDefinedPinDefaultValue(PinInfo, NewDefaultValue))
