K2Nodes now allow a CoreRedirect from the "self" pin if (and only if) the new function changes the visibility of the "self" pin.  This allows you to redirect member functions to non-member static functions (which will hide the self pin).  For example:

If the original call was a member function Foo.Bar(123) you could redirect that to a static function Bar(Foo, 123).

Syntax in BaseEngine.ini [[CoreRedirects] looks something like:
; Redirect the Member Function to a Static Function
+FunctionRedirects=(OldName="/Script/PluginName.Foo.Bar", NewName="FooBlueprintLibrary.Bar")
; Redirect the Self Pin to a Parameter Pin named "InFoo"
+PropertyRedirects=(OldName="FooBlueprintLibrary.Bar.Self", NewName="InFoo")

diff --git a/Engine/Source/Editor/BlueprintGraph/Private/K2Node.cpp b/Engine/Source/Editor/BlueprintGraph/Private/K2Node.cpp
--- a/Engine/Source/Editor/BlueprintGraph/Private/K2Node.cpp
+++ b/Engine/Source/Editor/BlueprintGraph/Private/K2Node.cpp
@@ -912,7 +912,16 @@ UK2Node::ERedirectType UK2Node::DoPinsMatchForReconstruction(const UEdGraphPin*
 		}
 		else
 		{
-			RedirectType = ERedirectType_Name;
+			// There is a special case for the 'self' pin because if you redirect a function, you potentially could also want to redirect the 'self' pin
+			// e.g. redirecting a member function to a static non-member function that takes the previous 'this' implicit parameter as an explicit parameter.
+			if (NewPin->PinName == UEdGraphSchema_K2::PN_Self)
+			{
+				RedirectType = (NewPin->bHidden == OldPin->bHidden) ? ERedirectType_Name : ERedirectType_None;
+			}
+			else
+			{
+				RedirectType = ERedirectType_Name;
+			}
 		}
 	}
 	else
diff --git a/Engine/Source/Editor/BlueprintGraph/Private/K2Node_Message.cpp b/Engine/Source/Editor/BlueprintGraph/Private/K2Node_Message.cpp
--- a/Engine/Source/Editor/BlueprintGraph/Private/K2Node_Message.cpp
+++ b/Engine/Source/Editor/BlueprintGraph/Private/K2Node_Message.cpp
@@ -179,12 +179,6 @@ void UK2Node_Message::ExpandNode(class FKismetCompilerContext& CompilerContext,
 			return;
 		}
 
-		if (!InterfaceClass->HasAnyClassFlags(CLASS_Interface))
-		{
-			CompilerContext.MessageLog.Error(*LOCTEXT("MessageNodeClassNotAnInterface_Error", "Message node @@ uses a class that isn't an interface.").ToString(), this);
-			return;
-		}
-
 		UFunction* MessageNodeFunction = GetTargetFunction();
 		if (MessageNodeFunction == nullptr)
 		{
@@ -198,8 +192,33 @@ void UK2Node_Message::ExpandNode(class FKismetCompilerContext& CompilerContext,
 			return;
 		}
 
+		const bool bIsStaticFunc = MessageNodeFunction->HasAllFunctionFlags(FUNC_Static);
+		if (!bIsStaticFunc && !InterfaceClass->HasAnyClassFlags(CLASS_Interface))
+		{
+			CompilerContext.MessageLog.Error(*LOCTEXT("MessageNodeClassNotAnInterface_Error", "Message node @@ uses a class that isn't an interface.").ToString(), this);
+			return;
+		}
+
 		// Check to make sure we have a target
 		UEdGraphPin* MessageSelfPin = Schema->FindSelfPin(*this, EGPD_Input);
+
+		// If we've intentionally hidden the 'self' pin (e.g. we're a static function), the pin can be redirected
+		bool bSelfPinCanBeRedirected = MessageSelfPin->bHidden;
+		if (bSelfPinCanBeRedirected)
+		{
+			// Build up a list of names that could correspond to this self pin (e.g. SomeClass.SomeFunc.self)
+			TArray<FString> PossibleRedirectNames;
+			GetRedirectPinNames(*MessageSelfPin, PossibleRedirectNames);
+
+			// Check for any redirects, and if we have one, try to find that pin
+			FName OutSelfPinName;
+			ERedirectType RedirectType = ShouldRedirectParam(PossibleRedirectNames, OutSelfPinName, this);
+			if (RedirectType != ERedirectType::ERedirectType_None)
+			{
+				MessageSelfPin = FindPin(OutSelfPinName, EEdGraphPinDirection::EGPD_Input);
+			}
+		}
+
 		if( !MessageSelfPin || MessageSelfPin->LinkedTo.Num() == 0 )
 		{
 			CompilerContext.MessageLog.Error(*LOCTEXT("MessageNodeSelfPin_Error", "Message node @@ must have a valid target or reference to self.").ToString(), this);
@@ -209,6 +228,18 @@ void UK2Node_Message::ExpandNode(class FKismetCompilerContext& CompilerContext,
 		// First, create an internal cast-to-interface node
 		UK2Node_DynamicCast* CastToInterfaceNode = CompilerContext.SpawnIntermediateNode<UK2Node_DynamicCast>(this, SourceGraph);
 		CastToInterfaceNode->TargetType = MessageNodeFunction->GetOuterUClass()->GetAuthoritativeClass();
+		if (bSelfPinCanBeRedirected)
+		{
+			if (UObject* PinSubCategoryObj = MessageSelfPin->PinType.PinSubCategoryObject.Get())
+			{
+				UClass* TargetTypeClass = Cast<UClass>(PinSubCategoryObj);
+				if (!TargetTypeClass)
+				{
+					TargetTypeClass = PinSubCategoryObj->GetClass();
+				}
+				CastToInterfaceNode->TargetType = TargetTypeClass->GetAuthoritativeClass();
+			}
+		}
 		CastToInterfaceNode->SetPurity(false);
 		CastToInterfaceNode->AllocateDefaultPins();
 
@@ -263,8 +294,13 @@ void UK2Node_Message::ExpandNode(class FKismetCompilerContext& CompilerContext,
 			LastOutCastSuccessPin = Schema->FindExecutionPin(*FunctionCallNode, EGPD_Output);
 		}
 		
-		// Self pin
-		UEdGraphPin* FunctionCallSelfPin = Schema->FindSelfPin(*FunctionCallNode, EGPD_Input);
+		// Function's actual 'self' pin (since Self Pin can be redirected, let's use its matching name)
+		UEdGraphPin* FunctionCallSelfPin = FunctionCallNode->FindPin(MessageSelfPin->GetName(), EGPD_Input);
+		if (!ensureMsgf(FunctionCallSelfPin, TEXT("Could not find suitable SelfPin '%s' matching input pin on (intermediate) node %s"), *MessageSelfPin->GetName(), *FunctionCallNode->GetDescriptiveCompiledName()))
+		{
+			// Fallback to old code with hardcoded Self name
+			FunctionCallSelfPin = Schema->FindSelfPin(*FunctionCallNode, EGPD_Input);
+		}
 		CastToInterfaceResultPin->MakeLinkTo(FunctionCallSelfPin);
 
 		UFunction* ArrayClearFunction = UKismetArrayLibrary::StaticClass()->FindFunctionByName(FName(TEXT("Array_Clear")));
