[Backout] - CL35773939

Partial back-out of this CL (MessageNodes still allow hookups to redirected static functions so as not to break potentially submitted content).

Original CL Desc
-----------------------------------------------------------------
K2Nodes now allow a CoreRedirect from the "self" pin if (and only if) the new function changes the visibility of the "self" pin.  This allows you to redirect member functions to non-member static functions (which will hide the self pin).  For example:

If the original call was a member function Foo.Bar(123) you could redirect that to a static function Bar(Foo, 123).

Syntax in BaseEngine.ini [[CoreRedirects] looks something like:
; Redirect the Member Function to a Static Function
+FunctionRedirects=(OldName="/Script/PluginName.Foo.Bar", NewName="FooBlueprintLibrary.Bar")
; Redirect the Self Pin to a Parameter Pin named "InFoo"
+PropertyRedirects=(OldName="FooBlueprintLibrary.Bar.Self", NewName="InFoo")

diff --git a/Engine/Source/Editor/BlueprintGraph/Private/K2Node.cpp b/Engine/Source/Editor/BlueprintGraph/Private/K2Node.cpp
--- a/Engine/Source/Editor/BlueprintGraph/Private/K2Node.cpp
+++ b/Engine/Source/Editor/BlueprintGraph/Private/K2Node.cpp
@@ -912,16 +912,7 @@ UK2Node::ERedirectType UK2Node::DoPinsMatchForReconstruction(const UEdGraphPin*
 		}
 		else
 		{
-			// There is a special case for the 'self' pin because if you redirect a function, you potentially could also want to redirect the 'self' pin
-			// e.g. redirecting a member function to a static non-member function that takes the previous 'this' implicit parameter as an explicit parameter.
-			if (NewPin->PinName == UEdGraphSchema_K2::PN_Self)
-			{
-				RedirectType = (NewPin->bHidden == OldPin->bHidden) ? ERedirectType_Name : ERedirectType_None;
-			}
-			else
-			{
-				RedirectType = ERedirectType_Name;
-			}
+			RedirectType = ERedirectType_Name;
 		}
 	}
 	else
