PR #12257: UK2Node_ConvertAsset improvement

Adds two functions to the UK2Node_ConvertAsset implementation that act as accessors to the pins created by that node.

diff --git a/Engine/Source/Editor/BlueprintGraph/Classes/K2Node_ConvertAsset.h b/Engine/Source/Editor/BlueprintGraph/Classes/K2Node_ConvertAsset.h
--- a/Engine/Source/Editor/BlueprintGraph/Classes/K2Node_ConvertAsset.h
+++ b/Engine/Source/Editor/BlueprintGraph/Classes/K2Node_ConvertAsset.h
@@ -34,6 +34,10 @@ public:
 	/** True if this is going from hard to soft, false for opposite or unknown */
 	bool IsConvertToSoft() const;
 
+	/** Pin Accessors */
+	UEdGraphPin* GetInputPin() const;
+	UEdGraphPin* GetOutputPin() const;
+
 	// UEdGraphNode interface
 	virtual void AllocateDefaultPins() override;
 	virtual FText GetNodeTitle(ENodeTitleType::Type TitleType) const override;
diff --git a/Engine/Source/Editor/BlueprintGraph/Private/K2Node_ConvertAsset.cpp b/Engine/Source/Editor/BlueprintGraph/Private/K2Node_ConvertAsset.cpp
--- a/Engine/Source/Editor/BlueprintGraph/Private/K2Node_ConvertAsset.cpp
+++ b/Engine/Source/Editor/BlueprintGraph/Private/K2Node_ConvertAsset.cpp
@@ -46,6 +46,16 @@ namespace UK2Node_ConvertAssetImpl
 	static const FName OutputPinName("Output");
 }
 
+UEdGraphPin* UK2Node_ConvertAsset::GetInputPin() const
+{
+	return FindPin(UK2Node_ConvertAssetImpl::InputPinName);
+}
+
+UEdGraphPin* UK2Node_ConvertAsset::GetOutputPin() const
+{
+	return FindPin(UK2Node_ConvertAssetImpl::OutputPinName);
+}
+
 UClass* UK2Node_ConvertAsset::GetTargetClass() const
 {
 	UEdGraphPin* InputPin = FindPin(UK2Node_ConvertAssetImpl::InputPinName);
