PR #12256: Native support for UK2Node_AsyncAction

Adds function to UK2Node_AsyncAction so that its members (really protected members of UK2Node_BaseAsyncTask) can be initialized properly from native. We can also do this via node spawners, but they tend to be clumsy when spawning async action nodes as intermediates.

#rnx

diff --git a/Engine/Source/Editor/BlueprintGraph/Classes/K2Node_AsyncAction.h b/Engine/Source/Editor/BlueprintGraph/Classes/K2Node_AsyncAction.h
--- a/Engine/Source/Editor/BlueprintGraph/Classes/K2Node_AsyncAction.h
+++ b/Engine/Source/Editor/BlueprintGraph/Classes/K2Node_AsyncAction.h
@@ -22,4 +22,7 @@ class BLUEPRINTGRAPH_API UK2Node_AsyncAction : public UK2Node_BaseAsyncTask
 	// UK2Node interface
 	virtual void GetMenuActions(FBlueprintActionDatabaseRegistrar& ActionRegistrar) const override;
 	// End of UK2Node interface
+
+	/** Initialize the async task from a known function **/
+	void InitializeProxyFromFunction(const UFunction* ProxyFunction);
 };
diff --git a/Engine/Source/Editor/BlueprintGraph/Private/K2Node_AsyncAction.cpp b/Engine/Source/Editor/BlueprintGraph/Private/K2Node_AsyncAction.cpp
--- a/Engine/Source/Editor/BlueprintGraph/Private/K2Node_AsyncAction.cpp
+++ b/Engine/Source/Editor/BlueprintGraph/Private/K2Node_AsyncAction.cpp
@@ -68,4 +68,18 @@ void UK2Node_AsyncAction::GetMenuActions(FBlueprintActionDatabaseRegistrar& Acti
 	}) );
 }
 
+void UK2Node_AsyncAction::InitializeProxyFromFunction(const UFunction* ProxyFunction)
+{
+	if (ensure(ProxyFunction))
+	{
+		const FObjectProperty* ReturnProp = CastField<FObjectProperty>(ProxyFunction->GetReturnProperty());
+		if (ensure(ReturnProp))
+		{
+			ProxyFactoryFunctionName = ProxyFunction->GetFName();
+			ProxyFactoryClass = ProxyFunction->GetOuterUClass();
+			ProxyClass = ReturnProp->PropertyClass;
+		}
+	}
+}
+
 #undef LOCTEXT_NAMESPACE
