PR #11943: Extend the K2Node_Event API to include a GetDelegatePin function

diff --git a/Engine/Source/Editor/BlueprintGraph/Classes/K2Node_Event.h b/Engine/Source/Editor/BlueprintGraph/Classes/K2Node_Event.h
--- a/Engine/Source/Editor/BlueprintGraph/Classes/K2Node_Event.h
+++ b/Engine/Source/Editor/BlueprintGraph/Classes/K2Node_Event.h
@@ -132,6 +132,9 @@ class UK2Node_Event : public UK2Node_EditablePinBase, public IK2Node_EventNodeIn
 	/** Helper function to identify if two Event nodes are the same */
 	static BLUEPRINTGRAPH_API bool AreEventNodesIdentical(const UK2Node_Event* InNodeA, const UK2Node_Event* InNodeB);
 
+	/** Accessor to the delegate pin **/
+	BLUEPRINTGRAPH_API UEdGraphPin* GetDelegatePin() const;
+
 protected:
 	void FixupEventReference(bool bForce = false);
 
diff --git a/Engine/Source/Editor/BlueprintGraph/Private/K2Node_Event.cpp b/Engine/Source/Editor/BlueprintGraph/Private/K2Node_Event.cpp
--- a/Engine/Source/Editor/BlueprintGraph/Private/K2Node_Event.cpp
+++ b/Engine/Source/Editor/BlueprintGraph/Private/K2Node_Event.cpp
@@ -963,6 +963,11 @@ void UK2Node_Event::FindDiffs(UEdGraphNode* OtherNode, struct FDiffResults& Resu
 	}
 }
 
+UEdGraphPin* UK2Node_Event::GetDelegatePin() const
+{
+	return FindPinChecked(DelegateOutputName);
+}
+
 bool UK2Node_Event::AreEventNodesIdentical(const UK2Node_Event* InNodeA, const UK2Node_Event* InNodeB)
 {
 	return InNodeA->EventReference.GetMemberName() == InNodeB->EventReference.GetMemberName()
