Set bAutomaticallyRegisterInputOnConstruction to true on every legacy input node as well to ensure that they get bound during compilation. This is required to allow the use of these legacy key nodes in widget blueprints after the changes made in 35418847 to automatically create the widget blueprint bindings during Blueprint compilation.

#rnx

diff --git a/Engine/Plugins/EnhancedInput/Source/InputBlueprintNodes/Private/K2Node_InputDebugKey.cpp b/Engine/Plugins/EnhancedInput/Source/InputBlueprintNodes/Private/K2Node_InputDebugKey.cpp
--- a/Engine/Plugins/EnhancedInput/Source/InputBlueprintNodes/Private/K2Node_InputDebugKey.cpp
+++ b/Engine/Plugins/EnhancedInput/Source/InputBlueprintNodes/Private/K2Node_InputDebugKey.cpp
@@ -15,6 +15,7 @@
 #include "Styling/AppStyle.h"
 #include "GameFramework/InputSettings.h"
 #include "Editor.h"								// for FEditorDelegates::OnEnableGestureRecognizerChanged
+#include "Blueprint/UserWidget.h"
 
 #include UE_INLINE_GENERATED_CPP_BY_NAME(K2Node_InputDebugKey)
 
@@ -346,6 +347,15 @@ void UK2Node_InputDebugKey::ExpandNode(FKismetCompilerContext& CompilerContext,
 			CompilerContext.MovePinLinksToIntermediate(*ActionValuePin, *InputDebugKeyEvent->FindPin(ActionPinName));
 		}
 	}
+
+	// Widget blueprints require the bAutomaticallyRegisterInputOnConstruction to be set to true in order to receive callbacks
+	CompilerContext.AddPostCDOCompiledStep([](const UObject::FPostCDOCompiledContext& Context, UObject* NewCDO)
+	{
+		if (UUserWidget* Widget = Cast<UUserWidget>(NewCDO))
+		{
+			Widget->bAutomaticallyRegisterInputOnConstruction = true;
+		}
+	});
 }
 
 void UK2Node_InputDebugKey::GetMenuActions(FBlueprintActionDatabaseRegistrar& ActionRegistrar) const
diff --git a/Engine/Source/Editor/BlueprintGraph/BlueprintGraph.Build.cs b/Engine/Source/Editor/BlueprintGraph/BlueprintGraph.Build.cs
--- a/Engine/Source/Editor/BlueprintGraph/BlueprintGraph.Build.cs
+++ b/Engine/Source/Editor/BlueprintGraph/BlueprintGraph.Build.cs
@@ -34,6 +34,7 @@ public class BlueprintGraph : ModuleRules
 				"ToolMenus",
 				"AssetTools",
 				"EditorSubsystem",
+				"UMG",
 			}
 		);
 
diff --git a/Engine/Source/Editor/BlueprintGraph/Classes/K2Node_InputAxisEvent.h b/Engine/Source/Editor/BlueprintGraph/Classes/K2Node_InputAxisEvent.h
--- a/Engine/Source/Editor/BlueprintGraph/Classes/K2Node_InputAxisEvent.h
+++ b/Engine/Source/Editor/BlueprintGraph/Classes/K2Node_InputAxisEvent.h
@@ -61,6 +61,7 @@ class BLUEPRINTGRAPH_API UK2Node_InputAxisEvent : public UK2Node_Event
 	virtual void GetMenuActions(FBlueprintActionDatabaseRegistrar& ActionRegistrar) const override;
 	virtual FText GetMenuCategory() const override;
 	virtual FBlueprintNodeSignature GetSignature() const override;
+	virtual void ExpandNode(class FKismetCompilerContext& CompilerContext, UEdGraph* SourceGraph) override;
 	//~ End UK2Node Interface
 
 	void Initialize(const FName AxisName);
diff --git a/Engine/Source/Editor/BlueprintGraph/Private/K2Node_InputAction.cpp b/Engine/Source/Editor/BlueprintGraph/Private/K2Node_InputAction.cpp
--- a/Engine/Source/Editor/BlueprintGraph/Private/K2Node_InputAction.cpp
+++ b/Engine/Source/Editor/BlueprintGraph/Private/K2Node_InputAction.cpp
@@ -38,6 +38,7 @@
 #include "UObject/ObjectVersion.h"
 #include "UObject/WeakObjectPtr.h"
 #include "UObject/WeakObjectPtrTemplates.h"
+#include "Blueprint/UserWidget.h"
 
 #define LOCTEXT_NAMESPACE "K2Node_InputAction"
 
@@ -236,6 +237,15 @@ void UK2Node_InputAction::ExpandNode(FKismetCompilerContext& CompilerContext, UE
 			CompilerContext.MovePinLinksToIntermediate(*FindPin(TEXT("Key")), *InputActionEvent->FindPin(TEXT("Key")));
 		}
 	}
+
+	// Widget blueprints require the bAutomaticallyRegisterInputOnConstruction to be set to true in order to receive callbacks
+	CompilerContext.AddPostCDOCompiledStep([](const UObject::FPostCDOCompiledContext& Context, UObject* NewCDO)
+	{
+		if (UUserWidget* Widget = Cast<UUserWidget>(NewCDO))
+		{
+			Widget->bAutomaticallyRegisterInputOnConstruction = true;
+		}
+	});
 }
 
 void UK2Node_InputAction::GetMenuActions(FBlueprintActionDatabaseRegistrar& ActionRegistrar) const
diff --git a/Engine/Source/Editor/BlueprintGraph/Private/K2Node_InputAxisEvent.cpp b/Engine/Source/Editor/BlueprintGraph/Private/K2Node_InputAxisEvent.cpp
--- a/Engine/Source/Editor/BlueprintGraph/Private/K2Node_InputAxisEvent.cpp
+++ b/Engine/Source/Editor/BlueprintGraph/Private/K2Node_InputAxisEvent.cpp
@@ -19,6 +19,7 @@
 #include "GameFramework/InputSettings.h"
 #include "HAL/PlatformCrt.h"
 #include "Internationalization/Internationalization.h"
+#include "KismetCompiler.h"
 #include "Kismet2/BlueprintEditorUtils.h"
 #include "Kismet2/CompilerResultsLog.h"
 #include "Misc/AssertionMacros.h"
@@ -27,6 +28,7 @@
 #include "Templates/SubclassOf.h"
 #include "UObject/Class.h"
 #include "UObject/ObjectVersion.h"
+#include "Blueprint/UserWidget.h"
 
 #define LOCTEXT_NAMESPACE "K2Node_InputAxisEvent"
 
@@ -213,4 +215,18 @@ FBlueprintNodeSignature UK2Node_InputAxisEvent::GetSignature() const
 	return NodeSignature;
 }
 
+void UK2Node_InputAxisEvent::ExpandNode(FKismetCompilerContext& CompilerContext, UEdGraph* SourceGraph)
+{
+	Super::ExpandNode(CompilerContext, SourceGraph);
+
+	// Widget blueprints require the bAutomaticallyRegisterInputOnConstruction to be set to true in order to receive callbacks
+	CompilerContext.AddPostCDOCompiledStep([](const UObject::FPostCDOCompiledContext& Context, UObject* NewCDO)
+		{
+			if (UUserWidget* Widget = Cast<UUserWidget>(NewCDO))
+			{
+				Widget->bAutomaticallyRegisterInputOnConstruction = true;
+			}
+		});
+}
+
 #undef LOCTEXT_NAMESPACE
diff --git a/Engine/Source/Editor/BlueprintGraph/Private/K2Node_InputKey.cpp b/Engine/Source/Editor/BlueprintGraph/Private/K2Node_InputKey.cpp
--- a/Engine/Source/Editor/BlueprintGraph/Private/K2Node_InputKey.cpp
+++ b/Engine/Source/Editor/BlueprintGraph/Private/K2Node_InputKey.cpp
@@ -15,6 +15,7 @@
 #include "Styling/AppStyle.h"
 #include "GameFramework/InputSettings.h"
 #include "Editor.h"								// for FEditorDelegates::OnEnableGestureRecognizerChanged
+#include "Blueprint/UserWidget.h"
 
 #define LOCTEXT_NAMESPACE "UK2Node_InputKey"
 
@@ -362,6 +363,15 @@ void UK2Node_InputKey::ExpandNode(FKismetCompilerContext& CompilerContext, UEdGr
 			CompilerContext.MovePinLinksToIntermediate(*FindPin(TEXT("Key")), *InputKeyEvent->FindPin(TEXT("Key")));
 		}
 	}
+
+	// Widget blueprints require the bAutomaticallyRegisterInputOnConstruction to be set to true in order to receive callbacks
+	CompilerContext.AddPostCDOCompiledStep([](const UObject::FPostCDOCompiledContext& Context, UObject* NewCDO)
+	{
+		if (UUserWidget* Widget = Cast<UUserWidget>(NewCDO))
+		{
+			Widget->bAutomaticallyRegisterInputOnConstruction = true;
+		}
+	});
 }
 
 void UK2Node_InputKey::GetMenuActions(FBlueprintActionDatabaseRegistrar& ActionRegistrar) const
diff --git a/Engine/Source/Editor/BlueprintGraph/Private/K2Node_InputTouch.cpp b/Engine/Source/Editor/BlueprintGraph/Private/K2Node_InputTouch.cpp
--- a/Engine/Source/Editor/BlueprintGraph/Private/K2Node_InputTouch.cpp
+++ b/Engine/Source/Editor/BlueprintGraph/Private/K2Node_InputTouch.cpp
@@ -35,6 +35,7 @@
 #include "UObject/TopLevelAssetPath.h"
 #include "UObject/WeakObjectPtr.h"
 #include "UObject/WeakObjectPtrTemplates.h"
+#include "Blueprint/UserWidget.h"
 
 UK2Node_InputTouch::UK2Node_InputTouch(const FObjectInitializer& ObjectInitializer)
 	: Super(ObjectInitializer)
@@ -288,4 +289,13 @@ void UK2Node_InputTouch::ExpandNode(FKismetCompilerContext& CompilerContext, UEd
 			CompilerContext.MovePinLinksToIntermediate(*GetFingerIndexPin(), *InputTouchEvent->FindPin(TEXT("FingerIndex")));
 		}
 	}
+
+	// Widget blueprints require the bAutomaticallyRegisterInputOnConstruction to be set to true in order to receive callbacks
+	CompilerContext.AddPostCDOCompiledStep([](const UObject::FPostCDOCompiledContext& Context, UObject* NewCDO)
+	{
+		if (UUserWidget* Widget = Cast<UUserWidget>(NewCDO))
+		{
+			Widget->bAutomaticallyRegisterInputOnConstruction = true;
+		}
+	});
 }
diff --git a/Engine/Source/Editor/UMGEditor/Private/WidgetBlueprintCompiler.cpp b/Engine/Source/Editor/UMGEditor/Private/WidgetBlueprintCompiler.cpp
--- a/Engine/Source/Editor/UMGEditor/Private/WidgetBlueprintCompiler.cpp
+++ b/Engine/Source/Editor/UMGEditor/Private/WidgetBlueprintCompiler.cpp
@@ -715,6 +715,10 @@ void FWidgetBlueprintCompilerContext::CopyTermDefaultsToDefaultObject(UObject* D
 		{
 			DefaultWidget->bHasScriptImplementedPaint = false;
 		}
+
+		// Reset the value of this flag, which is set on PostCDOCompiled if there are any input nodes
+		// in the widget graphs.
+		DefaultWidget->bAutomaticallyRegisterInputOnConstruction = false;
 	}
 
 
