Add the ability to set DeprecatedFunction="Note" (and DeprecatedVariable) metadata in C++ to lower the severity of a blueprint deprecation message.
Add a utility function and DeprecatedVariable static name to support this.
Use this to replace an old hack that ignored warnings for BeginSpawningActorFromBlueprint.

diff --git a/Engine/Source/Editor/BlueprintGraph/Classes/EdGraphSchema_K2.h b/Engine/Source/Editor/BlueprintGraph/Classes/EdGraphSchema_K2.h
--- a/Engine/Source/Editor/BlueprintGraph/Classes/EdGraphSchema_K2.h
+++ b/Engine/Source/Editor/BlueprintGraph/Classes/EdGraphSchema_K2.h
@@ -138,6 +138,9 @@ public:
 	// [FunctionMetadata] Indicates that the function is deprecated
 	static const FName MD_DeprecatedFunction;
 
+	// [FunctionMetadata] Indicates that the property is deprecated
+	static const FName MD_DeprecatedProperty;
+
 	// [FunctionMetadata] Supplies the custom message to use for deprecation
 	static const FName MD_DeprecationMessage;
 
diff --git a/Engine/Source/Editor/BlueprintGraph/Private/EdGraphSchema_K2.cpp b/Engine/Source/Editor/BlueprintGraph/Private/EdGraphSchema_K2.cpp
--- a/Engine/Source/Editor/BlueprintGraph/Private/EdGraphSchema_K2.cpp
+++ b/Engine/Source/Editor/BlueprintGraph/Private/EdGraphSchema_K2.cpp
@@ -183,6 +183,7 @@ const FName FBlueprintMetadata::MD_Variadic(TEXT("Variadic"));
 const FName FBlueprintMetadata::MD_UnsafeForConstructionScripts(TEXT("UnsafeDuringActorConstruction"));
 const FName FBlueprintMetadata::MD_FunctionCategory(TEXT("Category"));
 const FName FBlueprintMetadata::MD_DeprecatedFunction(TEXT("DeprecatedFunction"));
+const FName FBlueprintMetadata::MD_DeprecatedProperty(TEXT("DeprecatedProperty"));
 const FName FBlueprintMetadata::MD_DeprecationMessage(TEXT("DeprecationMessage"));
 const FName FBlueprintMetadata::MD_CompactNodeTitle(TEXT("CompactNodeTitle"));
 const FName FBlueprintMetadata::MD_DisplayName(TEXT("DisplayName"));
diff --git a/Engine/Source/Editor/BlueprintGraph/Private/K2Node_CallFunction.cpp b/Engine/Source/Editor/BlueprintGraph/Private/K2Node_CallFunction.cpp
--- a/Engine/Source/Editor/BlueprintGraph/Private/K2Node_CallFunction.cpp
+++ b/Engine/Source/Editor/BlueprintGraph/Private/K2Node_CallFunction.cpp
@@ -556,19 +556,15 @@ FEdGraphNodeDeprecationResponse UK2Node_CallFunction::GetDeprecationResponse(EEd
 	FEdGraphNodeDeprecationResponse Response = Super::GetDeprecationResponse(DeprecationType);
 	if (DeprecationType == EEdGraphNodeDeprecationType::NodeHasDeprecatedReference)
 	{
-		// TEMP: Do not warn in the case of SpawnActor, as we have a special upgrade path for those nodes
-		if (FunctionReference.GetMemberName() == FName(TEXT("BeginSpawningActorFromBlueprint")))
+		UFunction* Function = GetTargetFunction();
+		if (ensureMsgf(Function != nullptr, TEXT("This node should not be able to report having a deprecated reference if the target function cannot be resolved.")))
 		{
-			Response.MessageType = EEdGraphNodeDeprecationMessageType::None;
-		}
-		else
-		{
-			UFunction* Function = GetTargetFunction();
-			if (ensureMsgf(Function != nullptr, TEXT("This node should not be able to report having a deprecated reference if the target function cannot be resolved.")))
-			{
-				FString DetailedMessage = Function->GetMetaData(FBlueprintMetadata::MD_DeprecationMessage);
-				Response.MessageText = FBlueprintEditorUtils::GetDeprecatedMemberUsageNodeWarning(ObjectTools::GetUserFacingFunctionName(Function), FText::FromString(DetailedMessage));
-			}
+			// Check the deprecation type to override the severity
+			FString MessageType = Function->GetMetaData(FBlueprintMetadata::MD_DeprecatedFunction);
+			Response.MessageType = FBlueprintEditorUtils::GetDeprecatedMessageType(MessageType);
+
+			FString DetailedMessage = Function->GetMetaData(FBlueprintMetadata::MD_DeprecationMessage);
+			Response.MessageText = FBlueprintEditorUtils::GetDeprecatedMemberUsageNodeWarning(ObjectTools::GetUserFacingFunctionName(Function), FText::FromString(DetailedMessage));
 		}
 	}
 	
diff --git a/Engine/Source/Editor/BlueprintGraph/Private/K2Node_Event.cpp b/Engine/Source/Editor/BlueprintGraph/Private/K2Node_Event.cpp
--- a/Engine/Source/Editor/BlueprintGraph/Private/K2Node_Event.cpp
+++ b/Engine/Source/Editor/BlueprintGraph/Private/K2Node_Event.cpp
@@ -859,6 +859,10 @@ FEdGraphNodeDeprecationResponse UK2Node_Event::GetDeprecationResponse(EEdGraphNo
 			UFunction* Function = EventReference.ResolveMember<UFunction>(GetBlueprintClassFromNode());
 			if (ensureMsgf(Function != nullptr, TEXT("This node should not be able to report having a deprecated reference if the event override cannot be resolved.")))
 			{
+				// Check the deprecation type to override the severity
+				FString MessageType = Function->GetMetaData(FBlueprintMetadata::MD_DeprecatedFunction);
+				Response.MessageType = FBlueprintEditorUtils::GetDeprecatedMessageType(MessageType);
+
 				FText EventName = FText::FromName(GetFunctionName());
 				FText DetailedMessage = FText::FromString(Function->GetMetaData(FBlueprintMetadata::MD_DeprecationMessage));
 				Response.MessageText = FBlueprintEditorUtils::GetDeprecatedMemberUsageNodeWarning(EventName, DetailedMessage);
diff --git a/Engine/Source/Editor/BlueprintGraph/Private/K2Node_FunctionEntry.cpp b/Engine/Source/Editor/BlueprintGraph/Private/K2Node_FunctionEntry.cpp
--- a/Engine/Source/Editor/BlueprintGraph/Private/K2Node_FunctionEntry.cpp
+++ b/Engine/Source/Editor/BlueprintGraph/Private/K2Node_FunctionEntry.cpp
@@ -665,6 +665,10 @@ FEdGraphNodeDeprecationResponse UK2Node_FunctionEntry::GetDeprecationResponse(EE
 			UFunction* const Function = FunctionReference.ResolveMember<UFunction>(GetBlueprintClassFromNode());
 			if (ensureMsgf(Function != nullptr, TEXT("This node should not be able to report having a deprecated reference if the override function cannot be resolved.")))
 			{
+				// Check the deprecation type to override the severity
+				FString MessageType = Function->GetMetaData(FBlueprintMetadata::MD_DeprecatedFunction);
+				Response.MessageType = FBlueprintEditorUtils::GetDeprecatedMessageType(MessageType);
+
 				FText FunctionName = FText::FromName(FunctionReference.GetMemberName());
 				FText DetailedMessage = FText::FromString(Function->GetMetaData(FBlueprintMetadata::MD_DeprecationMessage));
 				Response.MessageText = FBlueprintEditorUtils::GetDeprecatedMemberUsageNodeWarning(FunctionName, DetailedMessage);
diff --git a/Engine/Source/Editor/BlueprintGraph/Private/K2Node_MCDelegate.cpp b/Engine/Source/Editor/BlueprintGraph/Private/K2Node_MCDelegate.cpp
--- a/Engine/Source/Editor/BlueprintGraph/Private/K2Node_MCDelegate.cpp
+++ b/Engine/Source/Editor/BlueprintGraph/Private/K2Node_MCDelegate.cpp
@@ -312,6 +312,10 @@ FEdGraphNodeDeprecationResponse UK2Node_BaseMCDelegate::GetDeprecationResponse(E
 	{
 		if (FProperty* DelegateProperty = DelegateReference.ResolveMember<FProperty>(GetBlueprintClassFromNode()))
 		{
+			// Check the deprecation type to override the severity
+			FString MessageType = DelegateProperty->GetMetaData(FBlueprintMetadata::MD_DeprecatedProperty);
+			Response.MessageType = FBlueprintEditorUtils::GetDeprecatedMessageType(MessageType);
+
 			FString DetailedMessage = DelegateProperty->GetMetaData(FBlueprintMetadata::MD_DeprecationMessage);
 			Response.MessageText = FBlueprintEditorUtils::GetDeprecatedMemberUsageNodeWarning(GetPropertyDisplayName(), FText::FromString(DetailedMessage));
 		}
diff --git a/Engine/Source/Editor/BlueprintGraph/Private/K2Node_Variable.cpp b/Engine/Source/Editor/BlueprintGraph/Private/K2Node_Variable.cpp
--- a/Engine/Source/Editor/BlueprintGraph/Private/K2Node_Variable.cpp
+++ b/Engine/Source/Editor/BlueprintGraph/Private/K2Node_Variable.cpp
@@ -1165,6 +1165,10 @@ FEdGraphNodeDeprecationResponse UK2Node_Variable::GetDeprecationResponse(EEdGrap
 	{
 		if (FProperty* VariableProperty = VariableReference.ResolveMember<FProperty>(GetBlueprintClassFromNode()))
 		{
+			// Check the deprecation type to override the severity
+			FString MessageType = VariableProperty->GetMetaData(FBlueprintMetadata::MD_DeprecatedProperty);
+			Response.MessageType = FBlueprintEditorUtils::GetDeprecatedMessageType(MessageType);
+
 			FText MemberName = FText::FromName(VariableReference.GetMemberName());
 			FText DetailedMessage = FText::FromString(VariableProperty->GetMetaData(FBlueprintMetadata::MD_DeprecationMessage));
 			Response.MessageText = FBlueprintEditorUtils::GetDeprecatedMemberUsageNodeWarning(MemberName, DetailedMessage);
diff --git a/Engine/Source/Editor/UnrealEd/Private/Kismet2/BlueprintEditorUtils.cpp b/Engine/Source/Editor/UnrealEd/Private/Kismet2/BlueprintEditorUtils.cpp
--- a/Engine/Source/Editor/UnrealEd/Private/Kismet2/BlueprintEditorUtils.cpp
+++ b/Engine/Source/Editor/UnrealEd/Private/Kismet2/BlueprintEditorUtils.cpp
@@ -10040,6 +10040,21 @@ FText FBlueprintEditorUtils::GetDeprecatedMemberUsageNodeWarning(const FText& Me
 	return FText::Format(LOCTEXT("DeprecatedMemberUsageNodeWarning", "@@: Usage of '{MemberName}' has been deprecated. {DetailedMessage}"), Args);
 }
 
+EEdGraphNodeDeprecationMessageType FBlueprintEditorUtils::GetDeprecatedMessageType(const FString& TypeString)
+{
+	if (TypeString.Equals(TEXT("None"), ESearchCase::IgnoreCase))
+	{
+		return EEdGraphNodeDeprecationMessageType::None;
+	}
+	else if (TypeString.Equals(TEXT("Note"), ESearchCase::IgnoreCase))
+	{
+		return EEdGraphNodeDeprecationMessageType::Note;
+	}
+	
+	// Default to warning
+	return EEdGraphNodeDeprecationMessageType::Warning;
+}
+
 UK2Node_FunctionResult* FBlueprintEditorUtils::FindOrCreateFunctionResultNode(UK2Node_EditablePinBase* InFunctionEntryNode)
 {
 	UK2Node_FunctionResult* FunctionResult = nullptr;
diff --git a/Engine/Source/Editor/UnrealEd/Public/Kismet2/BlueprintEditorUtils.h b/Engine/Source/Editor/UnrealEd/Public/Kismet2/BlueprintEditorUtils.h
--- a/Engine/Source/Editor/UnrealEd/Public/Kismet2/BlueprintEditorUtils.h
+++ b/Engine/Source/Editor/UnrealEd/Public/Kismet2/BlueprintEditorUtils.h
@@ -1941,6 +1941,13 @@ public:
 	 */
 	static UNREALED_API FText GetDeprecatedMemberUsageNodeWarning(const FText& MemberName, const FText& DetailedMessage);
 
+	/**
+	 * Returns the severity of a deprecation warning by parsing a string loaded from metadata like DeprecatedFunction="Note"
+	 *
+	 * @param TypeString		(Required) String to parse for a message type
+	 */
+	static UNREALED_API EEdGraphNodeDeprecationMessageType GetDeprecatedMessageType(const FString& TypeString);
+
 	/**
 	 * Remove overridden component templates from instance component handlers when a parent class disables editable when inherited boolean.
 	 */
diff --git a/Engine/Source/Runtime/Engine/Classes/Kismet/GameplayStatics.h b/Engine/Source/Runtime/Engine/Classes/Kismet/GameplayStatics.h
--- a/Engine/Source/Runtime/Engine/Classes/Kismet/GameplayStatics.h
+++ b/Engine/Source/Runtime/Engine/Classes/Kismet/GameplayStatics.h
@@ -59,7 +59,7 @@ class UGameplayStatics : public UBlueprintFunctionLibrary
 	// --- Spawning functions ------------------------------
 
 	/** Spawns an instance of a blueprint, but does not automatically run its construction script.  */
-	UFUNCTION(BlueprintCallable, Category="Spawning", meta=(WorldContext="WorldContextObject", UnsafeDuringActorConstruction = "true", BlueprintInternalUseOnly = "true", DeprecatedFunction, DeprecationMessage="Use BeginSpawningActorFromClass"))
+	UFUNCTION(BlueprintCallable, Category="Spawning", meta=(WorldContext="WorldContextObject", UnsafeDuringActorConstruction = "true", BlueprintInternalUseOnly = "true", DeprecatedFunction="Note", DeprecationMessage = "Use BeginSpawningActorFromClass"))
 	static ENGINE_API class AActor* BeginSpawningActorFromBlueprint(const UObject* WorldContextObject, const class UBlueprint* Blueprint, const FTransform& SpawnTransform, bool bNoCollisionFail);
 
 	/** Spawns an instance of an actor class, but does not automatically run its construction script.  */
