<saved by Perforce>

#rnx

diff --git a/Engine/Source/Editor/BlueprintGraph/Private/CallFunctionHandler.cpp b/Engine/Source/Editor/BlueprintGraph/Private/CallFunctionHandler.cpp
--- a/Engine/Source/Editor/BlueprintGraph/Private/CallFunctionHandler.cpp
+++ b/Engine/Source/Editor/BlueprintGraph/Private/CallFunctionHandler.cpp
@@ -537,9 +537,8 @@ void FKCHandler_CallFunction::CreateFunctionCallStatement(FKismetFunctionContext
 				// reading an interface ptr (16 bytes) into an object ptr (8 bytes) when we process the context opcode.
 				const bool bIsInterfaceContextTerm = Target && Target->AssociatedVarProperty && Target->AssociatedVarProperty->IsA<FInterfaceProperty>();
 				
-				UClass* FunctionOwnerClass = Function->GetOuterUClass();
-				check(FunctionOwnerClass);
-				const bool bIsInterfaceFunc = FunctionOwnerClass->HasAnyClassFlags(CLASS_Interface);
+				UClass* FunctionOwnerClass = Function->GetOwnerClass();
+				const bool bIsInterfaceFunc = FunctionOwnerClass ? FunctionOwnerClass->HasAnyClassFlags(CLASS_Interface) : false;
 
 				FBlueprintCompiledStatement& Statement = Context.AppendStatementForNode(Node);
 				Statement.FunctionToCall = Function;
