UE-226059 - Nodes within math expression sub-graph can be edited to show exec pins

We shouldn't enable the pure toggle during any type of debugging mode.

diff --git a/Engine/Source/Editor/BlueprintGraph/Private/K2Node_CallFunction.cpp b/Engine/Source/Editor/BlueprintGraph/Private/K2Node_CallFunction.cpp
--- a/Engine/Source/Editor/BlueprintGraph/Private/K2Node_CallFunction.cpp
+++ b/Engine/Source/Editor/BlueprintGraph/Private/K2Node_CallFunction.cpp
@@ -3741,7 +3741,7 @@ void UK2Node_CallFunction::GetNodeContextMenuActions(class UToolMenu* Menu, clas
 		);
 	}
 
-	if (CanToggleNodePurity())
+	if (CanToggleNodePurity() && !Context->bIsDebugging)
 	{
 		const FText HideExecPinsTitle = LOCTEXT("HideExecPins", "Hide Exec pins");
 		const FText ShowExecPinsTitle = LOCTEXT("ShowExecPins", "Show Exec pins");
