PR #12431: Update K2Node_SwitchEnum.h (add SetEnum to the DLL export table)

Epic note: we should consider adding a SetEnum function to INodeDependingOnEnumInterface.

diff --git a/Engine/Source/Editor/BlueprintGraph/Classes/K2Node_SwitchEnum.h b/Engine/Source/Editor/BlueprintGraph/Classes/K2Node_SwitchEnum.h
--- a/Engine/Source/Editor/BlueprintGraph/Classes/K2Node_SwitchEnum.h
+++ b/Engine/Source/Editor/BlueprintGraph/Classes/K2Node_SwitchEnum.h
@@ -72,7 +72,7 @@ class UK2Node_SwitchEnum : public UK2Node_Switch, public INodeDependingOnEnumInt
 	// End of UK2Node_Switch Interface
 	
 	/** Bind the switch to a named enum */
-	void SetEnum(UEnum* InEnum);
+	BLUEPRINTGRAPH_API void SetEnum(UEnum* InEnum);
 
 protected:
 	/** Helper method to set-up pins */
