Make sure that FTypePromotion::IsValidPromotion actually respects the IgnoreTypePromotion metadata. This fixes a bug where, if you had a "Conv_" autocast UFUNCTION which you did NOT want to use type promotion, you can actually still connect it to the node and it will follow the code path to create the connection with a new external node instead of trying to use the type promotion node.

diff --git a/Engine/Source/Editor/BlueprintGraph/Private/BlueprintTypePromotion.cpp b/Engine/Source/Editor/BlueprintGraph/Private/BlueprintTypePromotion.cpp
--- a/Engine/Source/Editor/BlueprintGraph/Private/BlueprintTypePromotion.cpp
+++ b/Engine/Source/Editor/BlueprintGraph/Private/BlueprintTypePromotion.cpp
@@ -118,7 +118,15 @@ bool FTypePromotion::IsValidPromotion(const FEdGraphPinType& A, const FEdGraphPi
 	{
 		const UEdGraphSchema_K2* K2Schema = GetDefault<UEdGraphSchema_K2>();
 
-		return K2Schema->SearchForAutocastFunction(A, B).IsSet();
+		const TOptional<UEdGraphSchema_K2::FSearchForAutocastFunctionResults> Res = K2Schema->SearchForAutocastFunction(A, B);
+		
+		const UFunction* AutoFunc =
+			Res.IsSet() ?
+			Res->FunctionOwner->FindFunctionByName(Res->TargetFunction) :
+			nullptr;
+
+		// We can only use this conversion function for promotion if it is not marked to ignore it
+		return AutoFunc != nullptr && !AutoFunc->HasMetaData(FBlueprintMetadata::MD_IgnoreTypePromotion);
 	}
 	else
 	{
