PR #11770: Add support for the AdvancedDisplay markup on properties when using K2Node_ConstructObjectFromClass

diff --git a/Engine/Source/Editor/BlueprintGraph/Private/K2Node_ConstructObjectFromClass.cpp b/Engine/Source/Editor/BlueprintGraph/Private/K2Node_ConstructObjectFromClass.cpp
--- a/Engine/Source/Editor/BlueprintGraph/Private/K2Node_ConstructObjectFromClass.cpp
+++ b/Engine/Source/Editor/BlueprintGraph/Private/K2Node_ConstructObjectFromClass.cpp
@@ -149,6 +149,16 @@ void UK2Node_ConstructObjectFromClass::CreatePinsForClass(UClass* InClass, TArra
 					OutClassPins->Add(Pin);
 				}
 
+				if (Property->HasAllPropertyFlags(CPF_AdvancedDisplay))
+				{
+					Pin->bAdvancedView = true;
+
+					if (AdvancedPinDisplay == ENodeAdvancedPins::NoPins)
+					{
+						AdvancedPinDisplay = ENodeAdvancedPins::Hidden;
+					}
+				}
+
 				if (ClassDefaultObject && K2Schema->PinDefaultValueIsEditable(*Pin))
 				{
 					FString DefaultValueAsString;
