Mark UK2Node_BaseAsyncTask::GetFactoryFunction as public so it can be used for external validation, and some header cleanup

diff --git a/Engine/Source/Editor/BlueprintGraph/Classes/K2Node_BaseAsyncTask.h b/Engine/Source/Editor/BlueprintGraph/Classes/K2Node_BaseAsyncTask.h
--- a/Engine/Source/Editor/BlueprintGraph/Classes/K2Node_BaseAsyncTask.h
+++ b/Engine/Source/Editor/BlueprintGraph/Classes/K2Node_BaseAsyncTask.h
@@ -44,13 +44,16 @@ struct FAsyncTaskPinRedirectMapInfo
 	TMap<FName, TArray<UClass*> > OldPinToProxyClassMap;
 };
 
-/** !!! The proxy object should have RF_StrongRefOnFrame flag. !!! */
-
+/**
+ * Base class used for a blueprint node that calls a function and provides asynchronous exec pins when the function is complete.
+ * The created proxy object should have RF_StrongRefOnFrame flag.
+ */
 UCLASS(Abstract)
 class BLUEPRINTGRAPH_API UK2Node_BaseAsyncTask : public UK2Node
 {
 	GENERATED_UCLASS_BODY()
 
+public:
 	// UEdGraphNode interface
 	virtual void AllocateDefaultPins() override;
 	virtual FText GetTooltipText() const override;
@@ -69,17 +72,14 @@ class BLUEPRINTGRAPH_API UK2Node_BaseAsyncTask : public UK2Node
 	virtual FText GetToolTipHeading() const override;
 	virtual void GetMenuActions(FBlueprintActionDatabaseRegistrar& ActionRegistrar) const override;
 	virtual FText GetMenuCategory() const override;
-	/** Whether or not two pins match for purposes of reconnection after reconstruction.  This allows pins that may have had their names changed via reconstruction to be matched to their old values on a node-by-node basis, if needed*/
+	virtual void GetRedirectPinNames(const UEdGraphPin& Pin, TArray<FString>& RedirectPinNames) const override;
 	virtual ERedirectType DoPinsMatchForReconstruction(const UEdGraphPin* NewPin, int32 NewPinIndex, const UEdGraphPin* OldPin, int32 OldPinIndex) const override;
 	// End of UK2Node interface
 
-protected:
-	// Returns the factory function (checked)
+	/** Returns the primary function that this node will call */
 	UFunction* GetFactoryFunction() const;
 
-	/** Determines what the possible redirect pin names are **/
-	virtual void GetRedirectPinNames(const UEdGraphPin& Pin, TArray<FString>& RedirectPinNames) const;
-
+protected:
 	/**
 	 * If a the DefaultToSelf pin exists then it needs an actual connection to get properly casted
 	 * during compilation.
@@ -93,7 +93,6 @@ protected:
 	 */
 	bool ExpandDefaultToSelfPin(FKismetCompilerContext& CompilerContext, UEdGraph* SourceGraph, UK2Node_CallFunction* IntermediateProxyNode);
 	
-protected:
 	// The name of the function to call to create a proxy object
 	UPROPERTY()
 	FName ProxyFactoryFunctionName;
@@ -135,7 +134,6 @@ protected:
 	static TMap<FName, FAsyncTaskPinRedirectMapInfo> AsyncTaskPinRedirectMap;
 	static bool bAsyncTaskPinRedirectMapInitialized;
 
-protected:
 	/** Expand out the logic to handle the delegate output pins */
 	virtual bool HandleDelegates(
 		const TArray<FBaseAsyncTaskHelper::FOutputPinAndLocalVariable>& VariableOutputs, UEdGraphPin* ProxyObjectPin,
