Added util for querying whether a BP function is allowed for a given set of permissions

#rnx

diff --git a/Engine/Source/Editor/BlueprintGraph/Private/BlueprintEditorSettings.cpp b/Engine/Source/Editor/BlueprintGraph/Private/BlueprintEditorSettings.cpp
--- a/Engine/Source/Editor/BlueprintGraph/Private/BlueprintEditorSettings.cpp
+++ b/Engine/Source/Editor/BlueprintGraph/Private/BlueprintEditorSettings.cpp
@@ -277,12 +277,17 @@ bool UBlueprintEditorSettings::IsClassPathAllowedOnPin(const FTopLevelAssetPath&
 
 bool UBlueprintEditorSettings::IsFunctionAllowed(const UBlueprint* InBlueprint, const FName FunctionName) const
 {
-	if (!FunctionPermissions.HasFiltering())
+	return IsFunctionAllowedForPermissions(InBlueprint, FunctionName, FunctionPermissions);
+}
+
+bool UBlueprintEditorSettings::IsFunctionAllowedForPermissions(const UBlueprint* InBlueprint, const FName FunctionName, const FPathPermissionList& InFunctionPermissions)
+{
+	if (!InFunctionPermissions.HasFiltering())
 	{
 		return true;
 	}
 
-	if (FunctionPermissions.PassesFilter(FunctionName))
+	if (InFunctionPermissions.PassesFilter(FunctionName))
 	{
 		return true;
 	}
@@ -293,7 +298,7 @@ bool UBlueprintEditorSettings::IsFunctionAllowed(const UBlueprint* InBlueprint,
 		{
 			if (const UFunction* NativeParentFunction = NativeParentClass->FindFunctionByName(FunctionName))
 			{
-				if (FunctionPermissions.PassesFilter(NativeParentFunction->GetPathName()))
+				if (InFunctionPermissions.PassesFilter(NativeParentFunction->GetPathName()))
 				{
 					return true;
 				}
@@ -302,4 +307,4 @@ bool UBlueprintEditorSettings::IsFunctionAllowed(const UBlueprint* InBlueprint,
 	}
 
 	return false;
-}
\ No newline at end of file
+}
diff --git a/Engine/Source/Editor/BlueprintGraph/Public/BlueprintEditorSettings.h b/Engine/Source/Editor/BlueprintGraph/Public/BlueprintEditorSettings.h
--- a/Engine/Source/Editor/BlueprintGraph/Public/BlueprintEditorSettings.h
+++ b/Engine/Source/Editor/BlueprintGraph/Public/BlueprintEditorSettings.h
@@ -389,6 +389,8 @@ public:
 	bool HasClassPathOnPinFiltering() const { return IsClassPathAllowedOnPinDelegates.Num() > 0; }
 
 	bool IsFunctionAllowed(const UBlueprint* InBlueprint, const FName FunctionName) const;
+
+	static bool IsFunctionAllowedForPermissions(const UBlueprint* InBlueprint, const FName FunctionName, const FPathPermissionList& InFunctionPermissions);
 	
 private:
 	/** All function permissions */
