Only force pin inference when connection change notifications fail to infer themselves

[RN] Fixed bug that caused MakeArray with wildcard pins not behaving correctly when using smart macro inference

diff --git a/Engine/Source/Editor/BlueprintGraph/Private/K2Node_MacroInstance.cpp b/Engine/Source/Editor/BlueprintGraph/Private/K2Node_MacroInstance.cpp
--- a/Engine/Source/Editor/BlueprintGraph/Private/K2Node_MacroInstance.cpp
+++ b/Engine/Source/Editor/BlueprintGraph/Private/K2Node_MacroInstance.cpp
@@ -818,11 +818,18 @@ void UK2Node_MacroInstance::SmartInferWildcardsImpl(const TArray<UEdGraphNode*>&
 	while(DirtyNodePins.Num())
 	{
 		TArray<TPair<UEdGraphNode*, UEdGraphPin*>> CurrentDirtyNodePins = MoveTemp(DirtyNodePins);
+		TArray<UEdGraphNode*> NodesToForceInference;
 		for(const TPair<UEdGraphNode*, UEdGraphPin*>& DirtyNodePin : CurrentDirtyNodePins)
 		{
 			// Any wildcard pins that are connected to this pin need to be
 			// inferred and marked dirty:
-			CastChecked<UK2Node>(DirtyNodePin.Key)->NotifyPinConnectionListChanged(DirtyNodePin.Value);
+			UK2Node* TypedNode = CastChecked<UK2Node>(DirtyNodePin.Key);
+			TypedNode->NotifyPinConnectionListChanged(DirtyNodePin.Value);
+			const int32 WildcardCount = CountWildcardPins(TypedNode);
+			if (WildcardCount != 0 && WildcardCount == WildcardCounts[TypedNode])
+			{
+				NodesToForceInference.AddUnique(TypedNode); 
+			}
 		}
 
 		const auto InferLinkedPins = [](UEdGraphPin * Pin, const UEdGraphPin * SourcePin, TArray<TPair<UEdGraphNode*, UEdGraphPin*>>&OutDirtyNodePins)
@@ -830,6 +837,27 @@ void UK2Node_MacroInstance::SmartInferWildcardsImpl(const TArray<UEdGraphNode*>&
 			TSet<UEdGraphPin*> ProcessedPins;
 			UE::Private::InferLinkedPinsImpl(Pin, SourcePin->PinType, OutDirtyNodePins, ProcessedPins);
 		};
+		for(UEdGraphNode* Node : NodesToForceInference)
+		{
+			// force inference, the node is refusing to propagate based on connections.
+			// We can directly infer wildcard connections from their connected non wildcard
+			// pins:
+			for(UEdGraphPin* Pin : Node->Pins)
+			{
+				if(FWildcardNodeUtils::HasAnyWildcards(Pin))
+				{
+					for(UEdGraphPin* LinkedPin : Pin->LinkedTo)
+					{
+						if (!FWildcardNodeUtils::HasAnyWildcards(LinkedPin) )
+						{
+							// infer and mark dirty:
+							InferLinkedPins(Pin, LinkedPin, DirtyNodePins);
+							break;
+						}
+					}
+				}
+			}
+		}
 
 		// look for pins that are now inferable:
 		for(const TPair<UEdGraphNode*, int32>& NodeWithCount : WildcardCounts)
@@ -845,8 +873,8 @@ void UK2Node_MacroInstance::SmartInferWildcardsImpl(const TArray<UEdGraphNode*>&
 					{
 						if(FWildcardNodeUtils::HasAnyWildcards(LinkedPin))
 						{
-							// infer and mark dirty:
-							InferLinkedPins(LinkedPin, Pin, DirtyNodePins);
+							// mark dirty:
+							DirtyNodePins.AddUnique({LinkedPin->GetOwningNode(), LinkedPin});
 						}
 					}
 				}
