[BlueprintGraph]
* Attempt at fixing msvc analyse warning

diff --git a/Engine/Source/Editor/BlueprintGraph/Private/EdGraphSchema_K2.cpp b/Engine/Source/Editor/BlueprintGraph/Private/EdGraphSchema_K2.cpp
--- a/Engine/Source/Editor/BlueprintGraph/Private/EdGraphSchema_K2.cpp
+++ b/Engine/Source/Editor/BlueprintGraph/Private/EdGraphSchema_K2.cpp
@@ -5688,7 +5688,7 @@ void UEdGraphSchema_K2::GetGraphDisplayInformation(const UEdGraph& Graph, /*out*
 	else
 	{
 		// This is essentially GetUserFacingFunctionName but without having an actual UFunction to pull metadata from (no localization)
-		const bool bShowFriendlyNames = bAllowFriendlyNames && (GEditor && GetDefault<UEditorStyleSettings>()->bShowFriendlyNames);
+		const bool bShowFriendlyNames = bAllowFriendlyNames ? (GEditor && GetDefault<UEditorStyleSettings>()->bShowFriendlyNames) : false;
 		if (bShowFriendlyNames)
 		{
 			DisplayInfo.DisplayName = FText::FromString(FName::NameToDisplayString(DisplayInfo.PlainName.ToString(), false));
