[BP]: Harvest metadata from user-defined function local variable descriptors at compile time to support details customizations that may expose it for editing.

Notes:
- This CL also includes a change to avoid BP compile-time read access validation for "get" nodes backed by a reference to a local variable.
- Reason: if a variable is copied from "class" to "function local" space within the Blueprint editor, its metadata goes with it. If the class variable is private (for example) then that metadata will also be set in the pasted descriptor for the new local variable, and after this change that metadata will now propagate to the generated function-local FProperty at compile time, which will fail as it assumes class scope.
- Since user workflows might rely on metadata being saved across copy/paste operations to/from the local variable section, I chose not to modify the metadata itself (it is also being serialized so would require fixup on load to remove).

diff --git a/Engine/Source/Editor/BlueprintGraph/Private/K2Node_VariableGet.cpp b/Engine/Source/Editor/BlueprintGraph/Private/K2Node_VariableGet.cpp
--- a/Engine/Source/Editor/BlueprintGraph/Private/K2Node_VariableGet.cpp
+++ b/Engine/Source/Editor/BlueprintGraph/Private/K2Node_VariableGet.cpp
@@ -427,36 +427,44 @@ void UK2Node_VariableGet::ValidateNodeDuringCompilation(FCompilerResultsLog& Mes
 	Super::ValidateNodeDuringCompilation(MessageLog);
 
 	// Some expansions, such as timelines, will create gets for non-blueprint visible properties, and we don't want to validate against that
-	if (!IsIntermediateNode())
+	if (IsIntermediateNode())
 	{
-		if (FProperty* Property = GetPropertyForVariable())
-		{
-			const FBlueprintEditorUtils::EPropertyReadableState PropertyReadableState = FBlueprintEditorUtils::IsPropertyReadableInBlueprint(GetBlueprint(), Property);
+		return;
+	}
 
-			if (PropertyReadableState != FBlueprintEditorUtils::EPropertyReadableState::Readable)
+	// The validation below does not apply to local variables; they are always readable within their context.
+	if (VariableReference.IsLocalScope())
+	{
+		return;
+	}
+
+	if (FProperty* Property = GetPropertyForVariable())
+	{
+		const FBlueprintEditorUtils::EPropertyReadableState PropertyReadableState = FBlueprintEditorUtils::IsPropertyReadableInBlueprint(GetBlueprint(), Property);
+
+		if (PropertyReadableState != FBlueprintEditorUtils::EPropertyReadableState::Readable)
+		{
+			FFormatNamedArguments Args;
+			if (UObject* Class = Property->GetOwner<UObject>())
 			{
-				FFormatNamedArguments Args;
-				if (UObject* Class = Property->GetOwner<UObject>())
-				{
-					Args.Add(TEXT("VariableName"), FText::AsCultureInvariant(FString::Printf(TEXT("%s.%s"), *Class->GetName(), *Property->GetName())));
-				}
-				else
-				{
-					Args.Add(TEXT("VariableName"), FText::AsCultureInvariant(Property->GetName()));
-				}
+				Args.Add(TEXT("VariableName"), FText::AsCultureInvariant(FString::Printf(TEXT("%s.%s"), *Class->GetName(), *Property->GetName())));
+			}
+			else
+			{
+				Args.Add(TEXT("VariableName"), FText::AsCultureInvariant(Property->GetName()));
+			}
 
-				if (PropertyReadableState == FBlueprintEditorUtils::EPropertyReadableState::NotBlueprintVisible)
-				{
-					MessageLog.Error(*FText::Format(LOCTEXT("UnableToGet_NotVisible", "{VariableName} is not blueprint visible (BlueprintReadOnly or BlueprintReadWrite). Please fix mark up or cease accessing as this will be made an error in a future release. @@"), Args).ToString(), this);
-				}
-				else if (PropertyReadableState == FBlueprintEditorUtils::EPropertyReadableState::Private)
-				{
-					MessageLog.Error(*FText::Format(LOCTEXT("UnableToGet_ReadOnly", "{VariableName} is private and not accessible in this context. Please fix mark up or cease accessing as this will be an error in a future release. @@"), Args).ToString(), this);
-				}
-				else
-				{
-					check(false);
-				}
+			if (PropertyReadableState == FBlueprintEditorUtils::EPropertyReadableState::NotBlueprintVisible)
+			{
+				MessageLog.Error(*FText::Format(LOCTEXT("UnableToGet_NotVisible", "{VariableName} is not blueprint visible (BlueprintReadOnly or BlueprintReadWrite). Please fix mark up or cease accessing as this will be made an error in a future release. @@"), Args).ToString(), this);
+			}
+			else if (PropertyReadableState == FBlueprintEditorUtils::EPropertyReadableState::Private)
+			{
+				MessageLog.Error(*FText::Format(LOCTEXT("UnableToGet_ReadOnly", "{VariableName} is private and not accessible in this context. Please fix mark up or cease accessing as this will be an error in a future release. @@"), Args).ToString(), this);
+			}
+			else
+			{
+				check(false);
 			}
 		}
 	}
diff --git a/Engine/Source/Editor/KismetCompiler/Private/KismetCompiler.cpp b/Engine/Source/Editor/KismetCompiler/Private/KismetCompiler.cpp
--- a/Engine/Source/Editor/KismetCompiler/Private/KismetCompiler.cpp
+++ b/Engine/Source/Editor/KismetCompiler/Private/KismetCompiler.cpp
@@ -1752,6 +1752,12 @@ FProperty* FKismetCompilerContext::CreateUserDefinedLocalVariableForFunction(con
 		NewProperty->SetMetaData(TEXT("Category"), *Variable.Category.ToString());
 		NewProperty->RepNotifyFunc = Variable.RepNotifyFunc;
 		NewProperty->SetPropertyFlags((EPropertyFlags)Variable.PropertyFlags);
+
+		// Set metadata on property (some customizations may expose it)
+		for (const FBPVariableMetaDataEntry& Entry : Variable.MetaDataArray)
+		{
+			NewProperty->SetMetaData(Entry.DataKey, *Entry.DataValue);
+		}
 	}
 	
 	return NewProperty;
