Prevent reentrancy into FBlueprintActionDatabase constructor, which could potentially overwrite its own delegate handlers
[RN] Fixed crash when changing levels after restoring an editor session that had open blueprint editors

diff --git a/Engine/Source/Editor/BlueprintGraph/Private/BlueprintActionDatabase.cpp b/Engine/Source/Editor/BlueprintGraph/Private/BlueprintActionDatabase.cpp
--- a/Engine/Source/Editor/BlueprintGraph/Private/BlueprintActionDatabase.cpp
+++ b/Engine/Source/Editor/BlueprintGraph/Private/BlueprintActionDatabase.cpp
@@ -1142,6 +1142,9 @@ FBlueprintActionDatabase& FBlueprintActionDatabase::Get()
 	if (DatabaseInst == nullptr)
 	{
 		DatabaseInst = new FBlueprintActionDatabase();
+		// Run init after constructing to avoid potentially creating multiple 
+		// instances of the BAD:
+		DatabaseInst->Init();
 	}
 
 	return *DatabaseInst;
@@ -1155,6 +1158,10 @@ FBlueprintActionDatabase* FBlueprintActionDatabase::TryGet()
 
 //------------------------------------------------------------------------------
 FBlueprintActionDatabase::FBlueprintActionDatabase()
+{
+}
+
+void FBlueprintActionDatabase::Init()
 {
 	RefreshAll();
 	OnAssetLoadedDelegateHandle = FCoreUObjectDelegates::OnAssetLoaded.AddStatic(&BlueprintActionDatabaseImpl::OnAssetLoaded);
diff --git a/Engine/Source/Editor/BlueprintGraph/Public/BlueprintActionDatabase.h b/Engine/Source/Editor/BlueprintGraph/Public/BlueprintActionDatabase.h
--- a/Engine/Source/Editor/BlueprintGraph/Public/BlueprintActionDatabase.h
+++ b/Engine/Source/Editor/BlueprintGraph/Public/BlueprintActionDatabase.h
@@ -213,6 +213,7 @@ public:
 private:
 	/** Private constructor for singleton purposes. */
 	FBlueprintActionDatabase();
+	void Init();
 
 	/**
 	 * 
