PR #12927: Fix compile error on Mac.

#rnx

diff --git a/Engine/Source/Editor/BlueprintGraph/Private/K2Node_SetFieldsInStruct.cpp b/Engine/Source/Editor/BlueprintGraph/Private/K2Node_SetFieldsInStruct.cpp
--- a/Engine/Source/Editor/BlueprintGraph/Private/K2Node_SetFieldsInStruct.cpp
+++ b/Engine/Source/Editor/BlueprintGraph/Private/K2Node_SetFieldsInStruct.cpp
@@ -65,7 +65,7 @@ public:
 
 	virtual UEdGraphPin* FindStructPinChecked(UEdGraphNode* InNode) const override
 	{
-		check(CastChecked<UK2Node_SetFieldsInStruct>(InNode));
+		check(InNode != nullptr && InNode->IsA<UK2Node_SetFieldsInStruct>());
 		UEdGraphPin* FoundPin = InNode->FindPinChecked(SetFieldsInStructHelper::StructRefPinName());
 		check(EGPD_Input == FoundPin->Direction);
 		return FoundPin;
