Allow querying whether a given function is allowed for a BP via a callback

#rnx

diff --git a/Engine/Source/Editor/BlueprintGraph/Private/BlueprintEditorSettings.cpp b/Engine/Source/Editor/BlueprintGraph/Private/BlueprintEditorSettings.cpp
--- a/Engine/Source/Editor/BlueprintGraph/Private/BlueprintEditorSettings.cpp
+++ b/Engine/Source/Editor/BlueprintGraph/Private/BlueprintEditorSettings.cpp
@@ -275,19 +275,41 @@ bool UBlueprintEditorSettings::IsClassPathAllowedOnPin(const FTopLevelAssetPath&
 	return true;
 }
 
-bool UBlueprintEditorSettings::IsFunctionAllowed(const UBlueprint* InBlueprint, const FName FunctionName) const
+void UBlueprintEditorSettings::RegisterIsFunctionAllowedForAssetDelegate(const FName OwnerName, FOnIsFunctionAllowedForAsset Delegate)
 {
-	return IsFunctionAllowedForPermissions(InBlueprint, FunctionName, FunctionPermissions);
+	IsFunctionAllowedForAssetDelegates.Add(OwnerName, MoveTemp(Delegate));
 }
 
-bool UBlueprintEditorSettings::IsFunctionAllowedForPermissions(const UBlueprint* InBlueprint, const FName FunctionName, const FPathPermissionList& InFunctionPermissions)
+void UBlueprintEditorSettings::UnregisterIsFunctionAllowedForAssetDelegate(const FName OwnerName)
+{
+	IsFunctionAllowedForAssetDelegates.Remove(OwnerName);
+}
+
+bool UBlueprintEditorSettings::IsFunctionAllowed(const UBlueprint* InBlueprint, const FName InFunctionName) const
+{
+	return IsFunctionAllowedForPermissions(InBlueprint, InFunctionName, FunctionPermissions);
+}
+
+bool UBlueprintEditorSettings::IsFunctionAllowedForAsset(const UBlueprint* InBlueprint, const FName InFunctionName) const
+{
+	for (const TPair<FName, FOnIsFunctionAllowedForAsset>& DelegatePair : IsFunctionAllowedForAssetDelegates)
+	{
+		if (!DelegatePair.Value.Execute(InBlueprint, InFunctionName))
+		{
+			return false;
+		}
+	}
+	return IsFunctionAllowed(InBlueprint, InFunctionName);
+}
+
+bool UBlueprintEditorSettings::IsFunctionAllowedForPermissions(const UBlueprint* InBlueprint, const FName InFunctionName, const FPathPermissionList& InFunctionPermissions)
 {
 	if (!InFunctionPermissions.HasFiltering())
 	{
 		return true;
 	}
 
-	if (InFunctionPermissions.PassesFilter(FunctionName))
+	if (InFunctionPermissions.PassesFilter(InFunctionName))
 	{
 		return true;
 	}
@@ -296,7 +318,7 @@ bool UBlueprintEditorSettings::IsFunctionAllowedForPermissions(const UBlueprint*
 	{
 		if (const UClass* NativeParentClass = FBlueprintEditorUtils::FindFirstNativeClass(InBlueprint->ParentClass))
 		{
-			if (const UFunction* NativeParentFunction = NativeParentClass->FindFunctionByName(FunctionName))
+			if (const UFunction* NativeParentFunction = NativeParentClass->FindFunctionByName(InFunctionName))
 			{
 				if (InFunctionPermissions.PassesFilter(NativeParentFunction->GetPathName()))
 				{
diff --git a/Engine/Source/Editor/BlueprintGraph/Public/BlueprintEditorSettings.h b/Engine/Source/Editor/BlueprintGraph/Public/BlueprintEditorSettings.h
--- a/Engine/Source/Editor/BlueprintGraph/Public/BlueprintEditorSettings.h
+++ b/Engine/Source/Editor/BlueprintGraph/Public/BlueprintEditorSettings.h
@@ -364,8 +364,9 @@ public:
 	/** Get allowed pin categories permissions list */
 	FNamePermissionList& GetPinCategoryPermissions() { return PinCategoryPermissions; }
 
-	DECLARE_DELEGATE_RetVal_OneParam(bool, FOnIsClassAllowed, const UClass* /*InClass*/)
-	DECLARE_DELEGATE_RetVal_OneParam(bool, FOnIsClassPathAllowed, const FTopLevelAssetPath& /*InClassPath*/)
+	using FOnIsClassAllowed = TDelegate<bool(const UClass* InClass)>;
+	using FOnIsClassPathAllowed = TDelegate<bool(const FTopLevelAssetPath& InClassPath)>;
+	using FOnIsFunctionAllowedForAsset = TDelegate<bool(const UBlueprint* InBlueprint, const FName InFunctionName)>;
 	
 	/** Delegates called to determine whether a class type is allowed */
 	void RegisterIsClassAllowedDelegate(const FName OwnerName, FOnIsClassAllowed Delegate);
@@ -388,7 +389,13 @@ public:
 	bool IsClassPathAllowedOnPin(const FTopLevelAssetPath& InClassPath) const;
 	bool HasClassPathOnPinFiltering() const { return IsClassPathAllowedOnPinDelegates.Num() > 0; }
 
-	bool IsFunctionAllowed(const UBlueprint* InBlueprint, const FName FunctionName) const;
+	/** Delegates called to determine whether a function is allowed for a given Blueprint asset */
+	void RegisterIsFunctionAllowedForAssetDelegate(const FName OwnerName, FOnIsFunctionAllowedForAsset Delegate);
+	void UnregisterIsFunctionAllowedForAssetDelegate(const FName OwnerName);
+
+	bool IsFunctionAllowed(const UBlueprint* InBlueprint, const FName InFunctionName) const;
+
+	bool IsFunctionAllowedForAsset(const UBlueprint* InBlueprint, const FName InFunctionName) const;
 
 	static bool IsFunctionAllowedForPermissions(const UBlueprint* InBlueprint, const FName FunctionName, const FPathPermissionList& InFunctionPermissions);
 	
@@ -413,6 +420,9 @@ private:
 	TMap<FName, FOnIsClassAllowed> IsClassAllowedOnPinDelegates;
 	TMap<FName, FOnIsClassPathAllowed> IsClassPathAllowedOnPinDelegates;
 
+	/** Delegates called to determine whether a function is allowed for a given Blueprint asset */
+	TMap<FName, FOnIsFunctionAllowedForAsset> IsFunctionAllowedForAssetDelegates;
+
 protected:
 	//~ Begin UObject Interface
 	virtual void PostInitProperties();
diff --git a/Engine/Source/Editor/Kismet/Private/BlueprintEditor.cpp b/Engine/Source/Editor/Kismet/Private/BlueprintEditor.cpp
--- a/Engine/Source/Editor/Kismet/Private/BlueprintEditor.cpp
+++ b/Engine/Source/Editor/Kismet/Private/BlueprintEditor.cpp
@@ -1910,7 +1910,7 @@ void FBlueprintEditor::EnsureBlueprintIsUpToDate(UBlueprint* BlueprintObj)
 		{
 			ExistingUCS->bAllowDeletion = false;
 		}
-		else if(GetDefault<UBlueprintEditorSettings>()->IsFunctionAllowed(BlueprintObj, UEdGraphSchema_K2::FN_UserConstructionScript))
+		else if(GetDefault<UBlueprintEditorSettings>()->IsFunctionAllowedForAsset(BlueprintObj, UEdGraphSchema_K2::FN_UserConstructionScript))
 		{
 			UEdGraph* UCSGraph = FBlueprintEditorUtils::CreateNewGraph(BlueprintObj, UEdGraphSchema_K2::FN_UserConstructionScript, UEdGraph::StaticClass(), UEdGraphSchema_K2::StaticClass());
 			FBlueprintEditorUtils::AddFunctionGraph(BlueprintObj, UCSGraph, /*bIsUserCreated=*/ false, AActor::StaticClass());
