Shutown crash fix - prevent refreshing the blueprint autocast map after engine shutdown

diff --git a/Engine/Source/Editor/BlueprintGraph/Private/EdGraphSchema_K2.cpp b/Engine/Source/Editor/BlueprintGraph/Private/EdGraphSchema_K2.cpp
--- a/Engine/Source/Editor/BlueprintGraph/Private/EdGraphSchema_K2.cpp
+++ b/Engine/Source/Editor/BlueprintGraph/Private/EdGraphSchema_K2.cpp
@@ -2387,6 +2387,11 @@ public:
 
 	void Refresh()
 	{
+		if(IsEngineExitRequested())
+		{
+			return;
+		}
+
 		TRACE_CPUPROFILER_EVENT_SCOPE(WILD_FAutocastFunctionMap::Refresh);
 
 #ifdef SCHEMA_K2_AUTOCASTFUNCTIONMAP_LOG_TIME
