[BlueprintGraph]

Fix crash at shutdown where module unloading will trigger FAutocastFunctionMap::OnModulesChanged which will attempt to call `GetDefault<UEdGraphSchema_K2>()` after `PurgeUObjects` has been called and the script types no longer exist.

#rnx

diff --git a/Engine/Source/Editor/BlueprintGraph/Private/EdGraphSchema_K2.cpp b/Engine/Source/Editor/BlueprintGraph/Private/EdGraphSchema_K2.cpp
--- a/Engine/Source/Editor/BlueprintGraph/Private/EdGraphSchema_K2.cpp
+++ b/Engine/Source/Editor/BlueprintGraph/Private/EdGraphSchema_K2.cpp
@@ -2524,20 +2524,28 @@ public:
 		}
 	}
 
-	FAutocastFunctionMap()
+	void RegisterDelegates()
 	{
-		Refresh();
-
 		OnReloadCompleteDelegateHandle = FCoreUObjectDelegates::ReloadCompleteDelegate.AddStatic(&FAutocastFunctionMap::OnReloadComplete);
-
 		OnModulesChangedDelegateHandle = FModuleManager::Get().OnModulesChanged().AddStatic(&OnModulesChanged);
+		FCoreDelegates::OnEnginePreExit.AddLambda([]() { FAutocastFunctionMap::Get().UnregisterDelegates(); });
 	}
 
-	~FAutocastFunctionMap()
+	void UnregisterDelegates()
 	{
 		FCoreUObjectDelegates::ReloadCompleteDelegate.Remove(OnReloadCompleteDelegateHandle);
+		FModuleManager::Get().OnModulesChanged().Remove(OnModulesChangedDelegateHandle);
+	}
 
-		FModuleManager::Get().OnModulesChanged().Remove(OnModulesChangedDelegateHandle); 
+	FAutocastFunctionMap()
+	{
+		Refresh();
+		RegisterDelegates();
+	}
+
+	~FAutocastFunctionMap()
+	{
+		UnregisterDelegates();
 	}
 };
 
