Added nullcheck when dropping an asset on the Blueprint Graph in case the asset fails to load.

diff --git a/Engine/Source/Editor/BlueprintGraph/Private/EdGraphSchema_K2.cpp b/Engine/Source/Editor/BlueprintGraph/Private/EdGraphSchema_K2.cpp
--- a/Engine/Source/Editor/BlueprintGraph/Private/EdGraphSchema_K2.cpp
+++ b/Engine/Source/Editor/BlueprintGraph/Private/EdGraphSchema_K2.cpp
@@ -5861,6 +5861,11 @@ void UEdGraphSchema_K2::DroppedAssetsOnGraph(const TArray<FAssetData>& Assets, c
 			FVector2D Position = GraphPosition + (AssetIdx * FVector2D(XOffset, 0.0f));
 
 			UObject* Asset = Assets[AssetIdx].GetAsset();
+			if (!Asset)
+			{
+				UE_LOG(LogBlueprint, Warning, TEXT("Failed to load Asset '%s'."), *Assets[AssetIdx].GetFullName());
+				continue;
+			}
 
 			UClass* AssetClass = Asset->GetClass();		
 			if (UBlueprint* BlueprintAsset = Cast<UBlueprint>(Asset))
