UE-261892 - Editor crash after disconnecting EnumToString node input from a wildcard output pin

When an EnumToString node is created, its widget creates a SGraphPinEnum that also generates a combo box of enumerator values based on the currently linked enum. If there's no attached enum, it creates a generic byte-based SGraphPin. The logic of SGraphPinEnum assumes that its corresponding pin remains invariant. However, that's not the case if a link is ever broken or changes enum type. As a result, we can trip various asserts in SGraphPinEnum that assume there's a valid UEnum.

We can fix this by forcing a node reconstruction if the linked enum ever changes. In fact, we have to do this when breaking a link because we need a SGraphPin as opposed to a SGraphPinEnum. Additionally, we need to trigger the reconstruction in NodeConnectionListChanged. NotifyPinConnectionListChanged isn't safe for calling ReconstructNode because it's frequently called in contexts that expect stable pins, which reconstruction will destroy.

diff --git a/Engine/Source/Editor/BlueprintGraph/Classes/K2Node_GetEnumeratorName.h b/Engine/Source/Editor/BlueprintGraph/Classes/K2Node_GetEnumeratorName.h
--- a/Engine/Source/Editor/BlueprintGraph/Classes/K2Node_GetEnumeratorName.h
+++ b/Engine/Source/Editor/BlueprintGraph/Classes/K2Node_GetEnumeratorName.h
@@ -26,7 +26,7 @@ class UK2Node_GetEnumeratorName : public UK2Node
 {
 	GENERATED_UCLASS_BODY()
 
-	static FName EnumeratorPinName;
+public:
 
 	//~ Begin UEdGraphNode Interface
 	virtual void AllocateDefaultPins() override;
@@ -44,14 +44,18 @@ class UK2Node_GetEnumeratorName : public UK2Node
 	virtual bool IsConnectionDisallowed(const UEdGraphPin* MyPin, const UEdGraphPin* OtherPin, FString& OutReason) const override;
 	virtual void EarlyValidation(class FCompilerResultsLog& MessageLog) const override;
 	virtual void PostReconstructNode() override;
-	virtual void NotifyPinConnectionListChanged(UEdGraphPin* Pin) override;
+	virtual void NodeConnectionListChanged() override;
 	virtual void GetMenuActions(FBlueprintActionDatabaseRegistrar& ActionRegistrar) const override;
 	virtual FText GetMenuCategory() const override;
 	virtual int32 GetNodeRefreshPriority() const override { return EBaseNodeRefreshPriority::Low_UsesDependentWildcard; }
 	//~ End UK2Node Interface
 
-	void UpdatePinType();
-	class UEnum* GetEnum() const;
+protected:
+
+	static FName EnumeratorPinName;
+
+private:
+
+	UEnum* GetConnectedEnum() const;
 	virtual FName GetFunctionName() const;
 };
-
diff --git a/Engine/Source/Editor/BlueprintGraph/Classes/K2Node_GetEnumeratorNameAsString.h b/Engine/Source/Editor/BlueprintGraph/Classes/K2Node_GetEnumeratorNameAsString.h
--- a/Engine/Source/Editor/BlueprintGraph/Classes/K2Node_GetEnumeratorNameAsString.h
+++ b/Engine/Source/Editor/BlueprintGraph/Classes/K2Node_GetEnumeratorNameAsString.h
@@ -22,6 +22,8 @@ class UK2Node_GetEnumeratorNameAsString : public UK2Node_GetEnumeratorName
 {
 	GENERATED_UCLASS_BODY()
 
+public:
+
 	//~ Begin UEdGraphNode Interface
 	virtual void AllocateDefaultPins() override;
 	virtual FText GetTooltipText() const override;
@@ -35,6 +37,8 @@ class UK2Node_GetEnumeratorNameAsString : public UK2Node_GetEnumeratorName
 	virtual bool ShouldDrawCompact() const override { return false; }
 	//~ End UK2Node Interface
 
+private:
+
 	virtual FName GetFunctionName() const override;
 };
 
diff --git a/Engine/Source/Editor/BlueprintGraph/Private/K2Node_GetEnumeratorName.cpp b/Engine/Source/Editor/BlueprintGraph/Private/K2Node_GetEnumeratorName.cpp
--- a/Engine/Source/Editor/BlueprintGraph/Private/K2Node_GetEnumeratorName.cpp
+++ b/Engine/Source/Editor/BlueprintGraph/Private/K2Node_GetEnumeratorName.cpp
@@ -37,8 +37,6 @@ UK2Node_GetEnumeratorName::UK2Node_GetEnumeratorName(const FObjectInitializer& O
 
 void UK2Node_GetEnumeratorName::AllocateDefaultPins()
 {
-	const UEdGraphSchema_K2* Schema = GetDefault<UEdGraphSchema_K2>();
-
 	CreatePin(EGPD_Input, UEdGraphSchema_K2::PC_Byte, EnumeratorPinName);
 	CreatePin(EGPD_Output, UEdGraphSchema_K2::PC_Name, UEdGraphSchema_K2::PN_ReturnValue);
 }
@@ -58,10 +56,10 @@ FText UK2Node_GetEnumeratorName::GetCompactNodeTitle() const
 	return NSLOCTEXT("K2Node", "CastSymbol", "\x2022");
 }
 
-UEnum* UK2Node_GetEnumeratorName::GetEnum() const
+UEnum* UK2Node_GetEnumeratorName::GetConnectedEnum() const
 {
 	const UEdGraphPin* InputPin = FindPinChecked(EnumeratorPinName);
-	const UEdGraphPin* EnumPin = InputPin->LinkedTo.Num() ? InputPin->LinkedTo[0] : InputPin;
+	const UEdGraphPin* EnumPin = InputPin->LinkedTo.Num() ? InputPin->LinkedTo[0] : nullptr;
 	return EnumPin ? Cast<UEnum>(EnumPin->PinType.PinSubCategoryObject.Get()) : nullptr;
 }
 
@@ -88,8 +86,8 @@ void UK2Node_GetEnumeratorName::EarlyValidation(class FCompilerResultsLog& Messa
 {
 	Super::EarlyValidation(MessageLog);
 
-	const UEnum* Enum = GetEnum();
-	if (nullptr == Enum)
+	const UEnum* Enum = GetConnectedEnum();
+	if (Enum == nullptr)
 	{
 		MessageLog.Error(*NSLOCTEXT("K2Node", "GetNumEnumEntries_NoIntput_Error", "@@ Must have non-default Enum input").ToString(), this);
 	}
@@ -118,31 +116,26 @@ FName UK2Node_GetEnumeratorName::GetFunctionName() const
 	return FunctionName;
 }
 
-void UK2Node_GetEnumeratorName::UpdatePinType()
-{
-	UEdGraphPin* EnumPin = FindPinChecked(EnumeratorPinName);
-	UEnum* Enum = GetEnum();
-	if (EnumPin->PinType.PinSubCategoryObject != Enum)
-	{
-		EnumPin->PinType.PinSubCategoryObject = Enum;
-		PinTypeChanged(EnumPin);
-	}
-}
-
 void UK2Node_GetEnumeratorName::PostReconstructNode()
 {
-	UpdatePinType();
+	UEdGraphPin* EnumPin = FindPinChecked(EnumeratorPinName);
+	EnumPin->PinType.PinSubCategoryObject = GetConnectedEnum();
 
 	Super::PostReconstructNode();
 }
 
-void UK2Node_GetEnumeratorName::NotifyPinConnectionListChanged(UEdGraphPin* Pin)
+void UK2Node_GetEnumeratorName::NodeConnectionListChanged()
 {
-	Super::NotifyPinConnectionListChanged(Pin);
+	Super::NodeConnectionListChanged();
 
-	if (Pin == FindPinChecked(EnumeratorPinName))
+	// The corresponding SGraphPinEnum will effectively "cache" enum info (eg: it assumes a specific enumerator count).
+	// If our connected enum type changes for any reason, then we need to rebuild the widget.
+
+	UEdGraphPin* EnumPin = FindPinChecked(EnumeratorPinName);
+	UEnum* Enum = GetConnectedEnum();
+	if (EnumPin->PinType.PinSubCategoryObject != Enum)
 	{
-		UpdatePinType();
+		ReconstructNode();
 	}
 }
 
@@ -150,8 +143,8 @@ void UK2Node_GetEnumeratorName::ExpandNode(class FKismetCompilerContext& Compile
 {
 	Super::ExpandNode(CompilerContext, SourceGraph);
 
-	UEnum* Enum = GetEnum();
-	if(NULL == Enum)
+	UEnum* Enum = GetConnectedEnum();
+	if (Enum == nullptr)
 	{
 		CompilerContext.MessageLog.Error(*NSLOCTEXT("K2Node", "GetEnumeratorNam_Error_MustHaveValidName", "@@ must have a valid enum defined").ToString(), this);
 		return;
@@ -166,18 +159,18 @@ void UK2Node_GetEnumeratorName::ExpandNode(class FKismetCompilerContext& Compile
 	CallGetName->AllocateDefaultPins();
 	check(CallGetName->IsNodePure());
 		
-	//OPUTPUT PIN
+	// OUTPUT PIN
 	UEdGraphPin* OrgReturnPin = FindPinChecked(UEdGraphSchema_K2::PN_ReturnValue);
 	UEdGraphPin* NewReturnPin = CallGetName->GetReturnValuePin();
 	check(NewReturnPin);
 	CompilerContext.MovePinLinksToIntermediate(*OrgReturnPin, *NewReturnPin);
 
-	//ENUM PIN
+	// ENUM PIN
 	UEdGraphPin* EnumPin = CallGetName->FindPinChecked(TEXT("Enum"));
 	Schema->TrySetDefaultObject(*EnumPin, Enum);
 	check(EnumPin->DefaultObject == Enum);
 
-	//VALUE PIN
+	// VALUE PIN
 	UEdGraphPin* OrgInputPin = FindPinChecked(EnumeratorPinName);
 	UEdGraphPin* IndexPin = CallGetName->FindPinChecked(TEXT("EnumeratorValue"));
 	check(EGPD_Input == IndexPin->Direction && UEdGraphSchema_K2::PC_Byte == IndexPin->PinType.PinCategory);
@@ -185,7 +178,7 @@ void UK2Node_GetEnumeratorName::ExpandNode(class FKismetCompilerContext& Compile
 
 	if (!IndexPin->LinkedTo.Num())
 	{
-		//MAKE LITERAL BYTE FROM LITERAL ENUM
+		// MAKE LITERAL BYTE FROM LITERAL ENUM
 		const FString EnumLiteral = IndexPin->GetDefaultAsString();
 		const int32 NumericValue = IntCastChecked<int32, int64>(Enum->GetValueByName(*EnumLiteral));
 		if (NumericValue == INDEX_NONE) 
diff --git a/Engine/Source/Editor/BlueprintGraph/Private/K2Node_GetEnumeratorNameAsString.cpp b/Engine/Source/Editor/BlueprintGraph/Private/K2Node_GetEnumeratorNameAsString.cpp
--- a/Engine/Source/Editor/BlueprintGraph/Private/K2Node_GetEnumeratorNameAsString.cpp
+++ b/Engine/Source/Editor/BlueprintGraph/Private/K2Node_GetEnumeratorNameAsString.cpp
@@ -22,8 +22,6 @@ UK2Node_GetEnumeratorNameAsString::UK2Node_GetEnumeratorNameAsString(const FObje
 
 void UK2Node_GetEnumeratorNameAsString::AllocateDefaultPins()
 {
-	const UEdGraphSchema_K2* Schema = GetDefault<UEdGraphSchema_K2>();
-
 	CreatePin(EGPD_Input, UEdGraphSchema_K2::PC_Byte, EnumeratorPinName);
 	CreatePin(EGPD_Output, UEdGraphSchema_K2::PC_String, UEdGraphSchema_K2::PN_ReturnValue);
 }
