BlueprintGraph: Added check if component is BlueprintSpawnable when dropped on graph.

#13038

diff --git a/Engine/Source/Editor/BlueprintGraph/Private/BlueprintComponentNodeSpawner.cpp b/Engine/Source/Editor/BlueprintGraph/Private/BlueprintComponentNodeSpawner.cpp
--- a/Engine/Source/Editor/BlueprintGraph/Private/BlueprintComponentNodeSpawner.cpp
+++ b/Engine/Source/Editor/BlueprintGraph/Private/BlueprintComponentNodeSpawner.cpp
@@ -113,12 +113,6 @@ UBlueprintComponentNodeSpawner* UBlueprintComponentNodeSpawner::Create(const FCo
 		return NodeSpawner;
 	}
 
-	if (!FKismetEditorUtilities::IsClassABlueprintSpawnableComponent(ComponentClass))
-	{
-		// loaded class that is marked as abstract or not spawnable, don't create an entry:
-		return nullptr;
-	}
-
 	if (ComponentClass->ClassWithin && ComponentClass->ClassWithin != UObject::StaticClass())
 	{
 		// we can't support 'Within' markup on components at this time (core needs to be aware of non-CDO archetypes
@@ -239,18 +233,22 @@ UEdGraphNode* UBlueprintComponentNodeSpawner::Invoke(UEdGraph* ParentGraph, FBin
 			NewNode->TemplateType = Class;
 		}
 
-		UBlueprint* Blueprint = NewNode->GetBlueprint();
+		UActorComponent* ComponentTemplate = nullptr;
+		if (FKismetEditorUtilities::IsClassABlueprintSpawnableComponent(Class))
+		{
+			UBlueprint* Blueprint = NewNode->GetBlueprint();
 
-		FName DesiredComponentName = NewNode->MakeNewComponentTemplateName(Blueprint->GeneratedClass, Class);
-		UActorComponent* ComponentTemplate = NewObject<UActorComponent>(Blueprint->GeneratedClass, Class, DesiredComponentName, RF_ArchetypeObject | RF_Public | RF_Transactional);
+			FName DesiredComponentName = NewNode->MakeNewComponentTemplateName(Blueprint->GeneratedClass, Class);
+			ComponentTemplate = NewObject<UActorComponent>(Blueprint->GeneratedClass, Class, DesiredComponentName, RF_ArchetypeObject | RF_Public | RF_Transactional);
 
-		Blueprint->ComponentTemplates.Add(ComponentTemplate);
+			Blueprint->ComponentTemplates.Add(ComponentTemplate);
+		}
 
 		// set the name of the template as the default for the TemplateName param
 		UEdGraphPin* TemplateNamePin = NewNode->GetTemplateNamePinChecked();
 		if (TemplateNamePin != nullptr)
 		{
-			TemplateNamePin->DefaultValue = ComponentTemplate->GetName();
+			TemplateNamePin->DefaultValue = ComponentTemplate ? ComponentTemplate->GetName() : TEXT("");
 		}
 		NewNode->ReconstructNode();
 	}
diff --git a/Engine/Source/Editor/BlueprintGraph/Private/EdGraphSchema_K2.cpp b/Engine/Source/Editor/BlueprintGraph/Private/EdGraphSchema_K2.cpp
--- a/Engine/Source/Editor/BlueprintGraph/Private/EdGraphSchema_K2.cpp
+++ b/Engine/Source/Editor/BlueprintGraph/Private/EdGraphSchema_K2.cpp
@@ -5895,10 +5895,16 @@ void UEdGraphSchema_K2::DroppedAssetsOnGraph(const TArray<FAssetData>& Assets, c
 				const FScopedTransaction Transaction(LOCTEXT("CreateAddComponentFromAsset", "Add Component From Asset"), UEdGraphSchemaImpl::ShouldActuallyTransact());
 
 				FComponentTypeEntry ComponentType = { FString(), FString(), DestinationComponentType.GetGCPtr() };
-
-				IBlueprintNodeBinder::FBindingSet Bindings;
-				Bindings.Add(Asset);
-				UBlueprintComponentNodeSpawner::Create(ComponentType)->Invoke(Graph, Bindings, GraphPosition);
+				if (const UBlueprintComponentNodeSpawner* Spawner = UBlueprintComponentNodeSpawner::Create(ComponentType))
+				{
+					IBlueprintNodeBinder::FBindingSet Bindings;
+					Bindings.Add(Asset);
+					Spawner->Invoke(Graph, Bindings, GraphPosition);
+				}
+				else
+				{
+					UE_LOG(LogBlueprint, Warning, TEXT("Component Node could not be created for dropped Asset '%s'."), *Asset->GetName());
+				}
 			}
 		}
 	}
diff --git a/Engine/Source/Editor/BlueprintGraph/Private/K2Node_AddComponent.cpp b/Engine/Source/Editor/BlueprintGraph/Private/K2Node_AddComponent.cpp
--- a/Engine/Source/Editor/BlueprintGraph/Private/K2Node_AddComponent.cpp
+++ b/Engine/Source/Editor/BlueprintGraph/Private/K2Node_AddComponent.cpp
@@ -574,7 +574,7 @@ void UK2Node_AddComponent::MakeNewComponentTemplate()
 		// Determine the type of the component needed
 		UClass* ComponentClass = Cast<UClass>(ReturnPin->PinType.PinSubCategoryObject.Get());
 
-		if (ComponentClass)
+		if (ComponentClass && FKismetEditorUtilities::IsClassABlueprintSpawnableComponent(ComponentClass))
 		{
 			ensure(nullptr != Cast<UBlueprintGeneratedClass>(Blueprint->GeneratedClass));
 
@@ -611,7 +611,6 @@ void UK2Node_AddComponent::MakeNewComponentTemplate()
 		else
 		{
 			// Clear the template connection; can't resolve the type of the component to create
-			ensure(false);
 			TemplateNamePin->DefaultValue = TEXT("");
 		}
 	}
