UE-221698 - Convert Bool to Branch similar to Validated Get

Getter nodes for bools can be converted to branches via the context menu. This is a minor quality of life issue that obviates the need for manually adding a separate branch node.

The implementation generalizes the bIsPureGet to an enum that represents either the 'validated object get' or a regular boolean.

This also removes some really old and irrelevant code that was intended to catch issues tied to UE-19729.

diff --git a/Engine/Source/Editor/BlueprintGraph/Classes/K2Node_VariableGet.h b/Engine/Source/Editor/BlueprintGraph/Classes/K2Node_VariableGet.h
--- a/Engine/Source/Editor/BlueprintGraph/Classes/K2Node_VariableGet.h
+++ b/Engine/Source/Editor/BlueprintGraph/Classes/K2Node_VariableGet.h
@@ -25,11 +25,21 @@ class UObject;
 struct FBPVariableDescription;
 struct FEdGraphPinType;
 
+UENUM()
+enum class EGetNodeVariation
+{
+	Pure,
+	ValidatedObject,
+	Branch,
+};
+
 UCLASS()
 class BLUEPRINTGRAPH_API UK2Node_VariableGet : public UK2Node_Variable
 {
 	GENERATED_UCLASS_BODY()
 
+public:
+
 	//~ Begin UObject Interface
 	virtual void Serialize(FArchive& Ar) override;
 	//~ End UObject Interface
@@ -44,40 +54,32 @@ class BLUEPRINTGRAPH_API UK2Node_VariableGet : public UK2Node_Variable
 	//~ End UEdGraphNode Interface
 
 	//~ Begin K2Node Interface
-	virtual bool IsNodePure() const override { return bIsPureGet; }
+	virtual bool IsNodePure() const override { return CurrentVariation == EGetNodeVariation::Pure; }
 	virtual void ReallocatePinsDuringReconstruction(TArray<UEdGraphPin*>& OldPins) override;
 	virtual bool ShouldShowNodeProperties() const override { return true; }
 	virtual class FNodeHandlingFunctor* CreateNodeHandler(class FKismetCompilerContext& CompilerContext) const override;
 	virtual void ExpandNode(class FKismetCompilerContext& CompilerContext, UEdGraph* SourceGraph) override;
 	//~ End K2Node Interface
 
-	static FText GetPropertyTooltip(FProperty const* VariableProperty);
-	static FText GetBlueprintVarTooltip(FBPVariableDescription const& VarDesc);
-
-	/**
-	 * Will change the node's purity, and reallocate pins accordingly (adding/
-	 * removing exec pins).
-	 * 
-	 * @param  bNewPurity  The new value for bIsPureCast.
-	 */
-	void SetPurity(bool bNewPurity);
+	static FText GetPropertyTooltip(const FProperty* VariableProperty);
+	static FText GetBlueprintVarTooltip(const FBPVariableDescription& VarDesc);
 
 private:
-	/** Checks if the pin type is valid to be a non-pure node */
-	static bool IsValidTypeForNonPure(const FEdGraphPinType& InPinType);
 
-	/** Adds pins required for the node to function in a non-pure manner */
-	void CreateNonPurePins(TArray<UEdGraphPin*>* InOldPinsPtr);
+	/** Adds pins required for the node to function in a impure manner */
+	void CreateImpurePins(TArray<UEdGraphPin*>* InOldPinsPtr);
 
 	/** Flips the node's purity (adding/removing exec pins as needed). */
-	void TogglePurity();
+	void TogglePurity(EGetNodeVariation BoundVariation);
 
 	/** Constructing FText strings can be costly, so we cache the node's title/tooltip */
 	FNodeTextCache CachedTooltip;
 	FNodeTextCache CachedNodeTitle;
 
-	/** TRUE if the node should function as a pure node, without exec pins */
 	UPROPERTY()
-	bool bIsPureGet;
+	bool bIsPureGet_DEPRECATED = true;
+
+	UPROPERTY()
+	EGetNodeVariation CurrentVariation;
 };
 
diff --git a/Engine/Source/Editor/BlueprintGraph/Private/K2Node_VariableGet.cpp b/Engine/Source/Editor/BlueprintGraph/Private/K2Node_VariableGet.cpp
--- a/Engine/Source/Editor/BlueprintGraph/Private/K2Node_VariableGet.cpp
+++ b/Engine/Source/Editor/BlueprintGraph/Private/K2Node_VariableGet.cpp
@@ -92,42 +92,57 @@ namespace K2Node_VariableGetImpl
 	 * @param  VarName	The name of the variable that the node represents.
 	 * @return A formatted text string, describing what the VariableGet node does.
 	 */
-	static FText GetBaseTooltip(FName VarName);
-}
+	static FText GetBaseTooltip(FName VarName)
+	{
+		FFormatNamedArguments Args;
+		Args.Add(TEXT("VarName"), FText::FromName(VarName));
 
-static FText K2Node_VariableGetImpl::GetBaseTooltip(FName VarName)
-{
-	FFormatNamedArguments Args;
-	Args.Add(TEXT("VarName"), FText::FromName(VarName));
+		return FText::Format(LOCTEXT("GetVariableTooltip", "Read the value of variable {VarName}"), Args);
+	}
+
+	static EGetNodeVariation GetNodeVariation(const FEdGraphPinType& InPinType)
+	{
+		EGetNodeVariation Result = EGetNodeVariation::Pure;
 
-	return FText::Format(LOCTEXT("GetVariableTooltip", "Read the value of variable {VarName}"), Args);
+		if (!InPinType.IsContainer())
+		{
+			if (InPinType.PinCategory == UEdGraphSchema_K2::PC_Boolean)
+			{
+				Result = EGetNodeVariation::Branch;
+			}
+			else if (InPinType.PinCategory == UEdGraphSchema_K2::PC_Object || InPinType.PinCategory == UEdGraphSchema_K2::PC_Class || InPinType.PinCategory == UEdGraphSchema_K2::PC_SoftObject || InPinType.PinCategory == UEdGraphSchema_K2::PC_SoftClass)
+			{
+				Result = EGetNodeVariation::ValidatedObject;
+			}
+		}
 
+		return Result;
+	}
 }
 
 UK2Node_VariableGet::UK2Node_VariableGet(const FObjectInitializer& ObjectInitializer)
 	: Super(ObjectInitializer)
-	, bIsPureGet(true)
+	, CurrentVariation(EGetNodeVariation::Pure)
 {
 }
 
-void UK2Node_VariableGet::CreateNonPurePins(TArray<UEdGraphPin*>* InOldPinsPtr)
+void UK2Node_VariableGet::CreateImpurePins(TArray<UEdGraphPin*>* InOldPinsPtr)
 {
-	const UEdGraphSchema_K2* K2Schema = Cast<UEdGraphSchema_K2>(GetSchema());
-	check(K2Schema != nullptr);
+	const UEdGraphSchema_K2* K2Schema = CastChecked<UEdGraphSchema_K2>(GetSchema());
 	if (!K2Schema->DoesGraphSupportImpureFunctions(GetGraph()))
 	{
-		bIsPureGet = true;
+		CurrentVariation = EGetNodeVariation::Pure;
 	}
 
-	if (!bIsPureGet)
+	if (CurrentVariation != EGetNodeVariation::Pure)
 	{
 		FEdGraphPinType PinType;
-		FProperty* VariableProperty = GetPropertyForVariable();
+		const FProperty* VariableProperty = GetPropertyForVariable();
 
 		// We need the pin's type, to both see if it's an array and if it is of the correct types to remain an impure node
 		if (VariableProperty)
 		{
-			K2Schema->ConvertPropertyToPinType(GetPropertyForVariable(), PinType);
+			K2Schema->ConvertPropertyToPinType(VariableProperty, PinType);
 		}
 		// If there is no property and we are given some old pins to look at, find the old value pin and use the type there
 		// This allows nodes to be pasted into other BPs without access to the property
@@ -143,35 +158,50 @@ void UK2Node_VariableGet::CreateNonPurePins(TArray<UEdGraphPin*>* InOldPinsPtr)
 					break;
 				}
 			}
-
 		}
 
-		if (IsValidTypeForNonPure(PinType))
+		const EGetNodeVariation SupportedVariation = K2Node_VariableGetImpl::GetNodeVariation(PinType);
+		if (SupportedVariation == EGetNodeVariation::Branch)
 		{
 			// Input - Execution Pin
 			CreatePin(EGPD_Input, UEdGraphSchema_K2::PC_Exec, UEdGraphSchema_K2::PN_Execute);
 
 			// Output - Execution Pins
 			UEdGraphPin* ValidPin = CreatePin(EGPD_Output, UEdGraphSchema_K2::PC_Exec, UEdGraphSchema_K2::PN_Then);
-			ValidPin->PinFriendlyName = LOCTEXT("Valid", "Is Valid");
+			check(ValidPin);
+			ValidPin->PinFriendlyName = LOCTEXT("True", "True");
 
 			UEdGraphPin* InvalidPin = CreatePin(EGPD_Output, UEdGraphSchema_K2::PC_Exec, UEdGraphSchema_K2::PN_Else);
-			InvalidPin->PinFriendlyName = LOCTEXT("Invalid", "Is Not Valid");
+			check(InvalidPin);
+			InvalidPin->PinFriendlyName = LOCTEXT("False", "False");
 		}
-		else
+		else if (SupportedVariation == EGetNodeVariation::ValidatedObject)
 		{
-			bIsPureGet = true;
+			// Input - Execution Pin
+			CreatePin(EGPD_Input, UEdGraphSchema_K2::PC_Exec, UEdGraphSchema_K2::PN_Execute);
+
+			// Output - Execution Pins
+			UEdGraphPin* ValidPin = CreatePin(EGPD_Output, UEdGraphSchema_K2::PC_Exec, UEdGraphSchema_K2::PN_Then);
+			check(ValidPin);
+			ValidPin->PinFriendlyName = LOCTEXT("Valid", "Is Valid");
+
+			UEdGraphPin* InvalidPin = CreatePin(EGPD_Output, UEdGraphSchema_K2::PC_Exec, UEdGraphSchema_K2::PN_Else);
+			check(InvalidPin);
+			InvalidPin->PinFriendlyName = LOCTEXT("Invalid", "Is Not Valid");
 		}
+		
+		// Note that the type can changed independently of a manual toggle (eg: changing the variable's type)
+		CurrentVariation = SupportedVariation;
 	}
 }
 
 void UK2Node_VariableGet::AllocateDefaultPins()
 {
-	if(GetVarName() != NAME_None)
+	if (GetVarName() != NAME_None)
 	{
-		CreateNonPurePins(nullptr);
+		CreateImpurePins(nullptr);
 
-		if(CreatePinForVariable(EGPD_Output))
+		if (CreatePinForVariable(EGPD_Output))
 		{
 			CreatePinForSelf();
 		}
@@ -182,13 +212,13 @@ void UK2Node_VariableGet::AllocateDefaultPins()
 
 void UK2Node_VariableGet::ReallocatePinsDuringReconstruction(TArray<UEdGraphPin*>& OldPins)
 {
-	if(GetVarName() != NAME_None)
+	if (GetVarName() != NAME_None)
 	{
-		CreateNonPurePins(&OldPins);
+		CreateImpurePins(&OldPins);
 
-		if(!CreatePinForVariable(EGPD_Output))
+		if (!CreatePinForVariable(EGPD_Output))
 		{
-			if(!RecreatePinForVariable(EGPD_Output, OldPins))
+			if (!RecreatePinForVariable(EGPD_Output, OldPins))
 			{
 				return;
 			}
@@ -294,20 +324,7 @@ FText UK2Node_VariableGet::GetNodeTitle(ENodeTitleType::Type TitleType) const
 	{
 		UEdGraphPin* Pin = Pins[PinIndex];
 
-		// The following code is to attempt to log info related to UE-19729
-		if (TitleType == ENodeTitleType::ListView)
-		{
-			if (UEdGraph* Graph = Cast<UEdGraph>(GetOuter()))
-			{
-				FString VariableName = GetVarNameString();
-				FString BlueprintPath = FBlueprintEditorUtils::FindBlueprintForGraph(Graph)->GetPathName();
-				FString SetupStyle = bIsPureGet? TEXT("pure") : TEXT("validated");
-				FString VariableResolves = (VariableReference.ResolveMember<FProperty>(GetBlueprintClassFromNode()) != nullptr)? TEXT("resolves") : TEXT("does not resolve");
-				checkf(Pin, TEXT("Get node for variable '%s' in Blueprint '%s' which is setup as %s and has %d pins. Variable %s"), *VariableName, *BlueprintPath, *SetupStyle, Pins.Num(), *VariableResolves);
-			}
-		}
-
-		if (Pin->Direction == EGPD_Output)
+		if (ensure(Pin) && Pin->Direction == EGPD_Output)
 		{
 			++NumOutputsFound;
 			OutputPinName = Pin->PinName;
@@ -333,45 +350,46 @@ FNodeHandlingFunctor* UK2Node_VariableGet::CreateNodeHandler(FKismetCompilerCont
 	return new FKCHandler_VariableGet(CompilerContext);
 }
 
-bool UK2Node_VariableGet::IsValidTypeForNonPure(const FEdGraphPinType& InPinType)
-{
-	return !InPinType.IsContainer() && (InPinType.PinCategory == UEdGraphSchema_K2::PC_Object || InPinType.PinCategory == UEdGraphSchema_K2::PC_Class || InPinType.PinCategory == UEdGraphSchema_K2::PC_SoftObject || InPinType.PinCategory == UEdGraphSchema_K2::PC_SoftClass);
-}
-
 void UK2Node_VariableGet::GetNodeContextMenuActions(UToolMenu* Menu, UGraphNodeContextMenuContext* Context) const
 {
 	Super::GetNodeContextMenuActions(Menu, Context);
 
+	check(Menu);
+	check(Context);
+
+	const UEdGraphSchema_K2* K2Schema = CastChecked<UEdGraphSchema_K2>(GetSchema());
 	const UEdGraphPin* ValuePin = GetValuePin();
-	if (ValuePin && IsValidTypeForNonPure(ValuePin->PinType))
+	
+	if (ValuePin)
 	{
+		const EGetNodeVariation SupportedVariation = K2Node_VariableGetImpl::GetNodeVariation(ValuePin->PinType);
+		const bool bCanShowToggleVariationMenu =
+			K2Schema->DoesGraphSupportImpureFunctions(GetGraph()) &&
+			(SupportedVariation != EGetNodeVariation::Pure) &&
+			!Context->bIsDebugging
+		;
+
+		if (bCanShowToggleVariationMenu)
 		{
 			FText MenuEntryTitle;
 			FText MenuEntryTooltip;
 
-			bool bCanTogglePurity = true;
-			auto CanExecutePurityToggle = [](bool const bInCanTogglePurity)->bool
-			{
-				return bInCanTogglePurity;
-			};
-
-			if (bIsPureGet)
+			if (CurrentVariation == EGetNodeVariation::Pure)
 			{
-				MenuEntryTitle   = LOCTEXT("ConvertToImpureGetTitle",   "Convert to Validated Get");
-				MenuEntryTooltip = LOCTEXT("ConvertToImpureGetTooltip", "Adds in branching execution pins so that you can separately handle when the returned value is valid/invalid.");
-
-				const UEdGraphSchema_K2* K2Schema = Cast<UEdGraphSchema_K2>(GetSchema());
-				check(K2Schema != nullptr);
-
-				bCanTogglePurity = K2Schema->DoesGraphSupportImpureFunctions(GetGraph());
-				if (!bCanTogglePurity)
+				if (SupportedVariation == EGetNodeVariation::ValidatedObject)
 				{
-					MenuEntryTooltip = LOCTEXT("CannotMakeImpureGetTooltip", "This graph does not support impure calls!");
+					MenuEntryTitle = LOCTEXT("ConvertToImpureGetObjectTitle", "Convert to Validated Get");
+					MenuEntryTooltip = LOCTEXT("ConvertToImpureGetObjectTooltip", "Adds in branching execution pins so that you can separately handle when the returned value is valid/invalid.");
+				}
+				else if (SupportedVariation == EGetNodeVariation::Branch)
+				{
+					MenuEntryTitle = LOCTEXT("ConvertToImpureGetBooleanTitle", "Convert to Branch");
+					MenuEntryTooltip = LOCTEXT("ConvertToImpureGetBooleanTooltip", "Adds in branching execution pins so that you can separately handle when the returned value is true/false.");
 				}
 			}
 			else
 			{
-				MenuEntryTitle   = LOCTEXT("ConvertToPureGetTitle",   "Convert to pure Get");
+				MenuEntryTitle = LOCTEXT("ConvertToPureGetTitle", "Convert to pure Get");
 				MenuEntryTooltip = LOCTEXT("ConvertToPureGetTooltip", "Removes the execution pins to make the node more versatile.");
 			}
 
@@ -382,44 +400,47 @@ void UK2Node_VariableGet::GetNodeContextMenuActions(UToolMenu* Menu, UGraphNodeC
 				MenuEntryTooltip,
 				FSlateIcon(),
 				FUIAction(
-				FExecuteAction::CreateUObject(const_cast<UK2Node_VariableGet*>(this), &UK2Node_VariableGet::TogglePurity),
-				FCanExecuteAction::CreateStatic(CanExecutePurityToggle, bCanTogglePurity && !Context->bIsDebugging),
-				FIsActionChecked()
+					FExecuteAction::CreateUObject(const_cast<UK2Node_VariableGet*>(this), &UK2Node_VariableGet::TogglePurity, SupportedVariation),
+					FCanExecuteAction(),
+					FIsActionChecked()
 				)
-				);
+			);
 		}
 	}
 }
 
-void UK2Node_VariableGet::TogglePurity()
+void UK2Node_VariableGet::TogglePurity(EGetNodeVariation BoundVariation)
 {
 	FText TransactionTitle;
-	if(!bIsPureGet)
+
+	EGetNodeVariation PendingVariation = EGetNodeVariation::Pure;
+	if (CurrentVariation == EGetNodeVariation::Pure)
 	{
-		TransactionTitle = LOCTEXT("TogglePureGet", "Convert to Pure Get");
+		if (BoundVariation == EGetNodeVariation::ValidatedObject)
+		{
+			TransactionTitle = LOCTEXT("ToggleImpureGetObject", "Convert to Validated Get");
+		}
+		else if (BoundVariation == EGetNodeVariation::Branch)
+		{
+			TransactionTitle = LOCTEXT("ToggleImpureGetBoolean", "Convert to Branch");
+		}
+		else
+		{
+			check(false);
+		}
+
+		PendingVariation = BoundVariation;
 	}
 	else
 	{
-		TransactionTitle = LOCTEXT("ToggleImpureGet", "Convert to Impure Get");
+		TransactionTitle = LOCTEXT("TogglePureGet", "Convert to Pure Get");
+		PendingVariation = EGetNodeVariation::Pure;
 	}
-	const FScopedTransaction Transaction( TransactionTitle );
-	Modify();
-
-	SetPurity(!bIsPureGet);
-}
-
-void UK2Node_VariableGet::SetPurity(bool bNewPurity)
-{
-	if (bNewPurity != bIsPureGet)
-	{
-		bIsPureGet = bNewPurity;
 
-		bool const bHasBeenConstructed = (Pins.Num() > 0);
-		if (bHasBeenConstructed)
-		{
-			ReconstructNode();
-		}
-	}
+	const FScopedTransaction Transaction(TransactionTitle);
+	Modify();
+	CurrentVariation = PendingVariation;
+	ReconstructNode();
 }
 
 void UK2Node_VariableGet::ValidateNodeDuringCompilation(FCompilerResultsLog& MessageLog) const
@@ -476,20 +497,25 @@ void UK2Node_VariableGet::ExpandNode(class FKismetCompilerContext& CompilerConte
 
 	FProperty* VariableProperty = GetPropertyForVariable();
 
-	UK2Node_VariableGet* VariableGetNode = this;
-
 	// Do not attempt to expand the node when not a pure get nor when there is no property. Normal compilation error detection will detect the missing property.
-	if (!bIsPureGet && VariableProperty)
+	if ((CurrentVariation != EGetNodeVariation::Pure) && VariableProperty)
 	{
 		UEdGraphPin* ValuePin = GetValuePin();
+		check(ValuePin);
 
-		// Impure Get nodes convert into three nodes:
+		// Impure nodes need 2-3 intermediate nodes depending on the variation.
+		// 
+		// For validated objects, we need:
 		// 1. A pure Get node
 		// 2. An IsValid node
 		// 3. A Branch node (only impure part)
-		
+		//
+		// For branches, we only need:
+		// 1. A pure Get node
+		// 2. A Branch node (only impure part)
+
 		// Create the pure Get node
-		VariableGetNode = CompilerContext.SpawnIntermediateNode<UK2Node_VariableGet>(this, SourceGraph);
+		UK2Node_VariableGet* VariableGetNode = CompilerContext.SpawnIntermediateNode<UK2Node_VariableGet>(this, SourceGraph);
 		VariableGetNode->VariableReference = VariableReference;
 		VariableGetNode->AllocateDefaultPins();
 		CompilerContext.MessageLog.NotifyIntermediateObjectCreation(VariableGetNode, this);
@@ -501,33 +527,41 @@ void UK2Node_VariableGet::ExpandNode(class FKismetCompilerContext& CompilerConte
 			CompilerContext.MovePinLinksToIntermediate(*FindPin(UEdGraphSchema_K2::PN_Self), *VariableGetNode->FindPin(UEdGraphSchema_K2::PN_Self));
 		}
 
-		// Create the IsValid node
-		UK2Node_CallFunction* IsValidFunction = CompilerContext.SpawnIntermediateNode<UK2Node_CallFunction>(this, SourceGraph);
+		// By default, we'll assume that this is the branch variation.
+		// Otherwise, the validated object variation will change the source pin.
+		UEdGraphPin* SourceBoolPin = VariableGetNode->GetValuePin();
 
-		// Based on if the type is an "Object" or a "Class" changes which function to use
-		if (ValuePin->PinType.PinCategory == UEdGraphSchema_K2::PC_Object)
+		if (CurrentVariation == EGetNodeVariation::ValidatedObject)
 		{
-			IsValidFunction->SetFromFunction(UKismetSystemLibrary::StaticClass()->FindFunctionByName(GET_MEMBER_NAME_CHECKED(UKismetSystemLibrary, IsValid)));
-		}
-		else if (ValuePin->PinType.PinCategory == UEdGraphSchema_K2::PC_Class)
-		{
-			IsValidFunction->SetFromFunction(UKismetSystemLibrary::StaticClass()->FindFunctionByName(GET_MEMBER_NAME_CHECKED(UKismetSystemLibrary, IsValidClass)));
-		}
-		else if (ValuePin->PinType.PinCategory == UEdGraphSchema_K2::PC_SoftObject)
-		{
-			IsValidFunction->SetFromFunction(UKismetSystemLibrary::StaticClass()->FindFunctionByName(GET_MEMBER_NAME_CHECKED(UKismetSystemLibrary, IsValidSoftObjectReference)));
-		}
-		else if (ValuePin->PinType.PinCategory == UEdGraphSchema_K2::PC_SoftClass)
-		{
-			IsValidFunction->SetFromFunction(UKismetSystemLibrary::StaticClass()->FindFunctionByName(GET_MEMBER_NAME_CHECKED(UKismetSystemLibrary, IsValidSoftClassReference)));
-		}
-		IsValidFunction->AllocateDefaultPins();
-		CompilerContext.MessageLog.NotifyIntermediateObjectCreation(IsValidFunction, this);
+			// Create the IsValid node
+			UK2Node_CallFunction* IsValidFunction = CompilerContext.SpawnIntermediateNode<UK2Node_CallFunction>(this, SourceGraph);
 
-		// Connect the value pin from the new Get node to the IsValid
-		UEdGraphPin* ObjectPin = IsValidFunction->Pins[1];
-		check(ObjectPin->Direction == EGPD_Input);
-		ObjectPin->MakeLinkTo(VariableGetNode->GetValuePin());
+			// Based on if the type is an "Object" or a "Class" changes which function to use
+			if (ValuePin->PinType.PinCategory == UEdGraphSchema_K2::PC_Object)
+			{
+				IsValidFunction->SetFromFunction(UKismetSystemLibrary::StaticClass()->FindFunctionByName(GET_MEMBER_NAME_CHECKED(UKismetSystemLibrary, IsValid)));
+			}
+			else if (ValuePin->PinType.PinCategory == UEdGraphSchema_K2::PC_Class)
+			{
+				IsValidFunction->SetFromFunction(UKismetSystemLibrary::StaticClass()->FindFunctionByName(GET_MEMBER_NAME_CHECKED(UKismetSystemLibrary, IsValidClass)));
+			}
+			else if (ValuePin->PinType.PinCategory == UEdGraphSchema_K2::PC_SoftObject)
+			{
+				IsValidFunction->SetFromFunction(UKismetSystemLibrary::StaticClass()->FindFunctionByName(GET_MEMBER_NAME_CHECKED(UKismetSystemLibrary, IsValidSoftObjectReference)));
+			}
+			else if (ValuePin->PinType.PinCategory == UEdGraphSchema_K2::PC_SoftClass)
+			{
+				IsValidFunction->SetFromFunction(UKismetSystemLibrary::StaticClass()->FindFunctionByName(GET_MEMBER_NAME_CHECKED(UKismetSystemLibrary, IsValidSoftClassReference)));
+			}
+			IsValidFunction->AllocateDefaultPins();
+			CompilerContext.MessageLog.NotifyIntermediateObjectCreation(IsValidFunction, this);
+
+			// Connect the value pin from the new Get node to the IsValid
+			UEdGraphPin* ObjectPin = IsValidFunction->Pins[1];
+			check(ObjectPin->Direction == EGPD_Input);
+			ObjectPin->MakeLinkTo(VariableGetNode->GetValuePin());
+			SourceBoolPin = IsValidFunction->Pins[2];
+		}
 
 		// Create the Branch node
 		UK2Node_IfThenElse* BranchNode = CompilerContext.SpawnIntermediateNode<UK2Node_IfThenElse>(this, SourceGraph);
@@ -535,9 +569,8 @@ void UK2Node_VariableGet::ExpandNode(class FKismetCompilerContext& CompilerConte
 		CompilerContext.MessageLog.NotifyIntermediateObjectCreation(BranchNode, this);
 
 		// Connect the bool output pin from IsValid node to the Branch node
-		UEdGraphPin* BoolPin = IsValidFunction->Pins[2];
-		check(BoolPin->Direction == EGPD_Output);
-		BoolPin->MakeLinkTo(BranchNode->GetConditionPin());
+		check(SourceBoolPin->Direction == EGPD_Output);
+		SourceBoolPin->MakeLinkTo(BranchNode->GetConditionPin());
 
 		// Connect the Branch node to the input of the impure Get node
 		CompilerContext.MovePinLinksToIntermediate(*GetExecPin(), *BranchNode->GetExecPin());
@@ -580,27 +613,15 @@ void UK2Node_VariableGet::ExpandNode(class FKismetCompilerContext& CompilerConte
 
 void UK2Node_VariableGet::Serialize(FArchive& Ar)
 {
-	// The following code is to attempt to log info related to UE-19729
-	if (Ar.IsSaving() && Ar.IsPersistent())
+	Super::Serialize(Ar);
+
+	// Prior to the addition of CurrentVariation, bIsPureGet implied a 'validate object' variation.
+	// We also reset bIsPureGet to its default value to prevent triggering this data migration path again.
+	if (!bIsPureGet_DEPRECATED && Ar.IsLoading())
 	{
-		uint32 PortFlagsToSkip = PPF_Duplicate | PPF_DuplicateForPIE;
-		if (!(Ar.GetPortFlags() & PortFlagsToSkip))
-		{
-			if (UEdGraph* Graph = Cast<UEdGraph>(GetOuter()))
-			{
-				if (UBlueprint* Blueprint = FBlueprintEditorUtils::FindBlueprintForGraph(Graph))
-				{
-					if (!Blueprint->bBeingCompiled)
-					{
-						// The following line may spur the crash noted in UE-19729 and will confirm that the crash happens before the FiB gather.
-						GetNodeTitle(ENodeTitleType::ListView);
-					}
-				}
-			}
-		}
+		CurrentVariation = EGetNodeVariation::ValidatedObject;
+		bIsPureGet_DEPRECATED = true;
 	}
-	Super::Serialize(Ar);
 }
 
-
 #undef LOCTEXT_NAMESPACE
