[Truncation Warnings] Update EdGraphSchema, EdGraphActions, and SGraphEditor to use FVector2f

diff --git a/Engine/Plugins/Runtime/RigVM/Source/RigVMEditor/Private/Editor/RigVMEditor.cpp b/Engine/Plugins/Runtime/RigVM/Source/RigVMEditor/Private/Editor/RigVMEditor.cpp
--- a/Engine/Plugins/Runtime/RigVM/Source/RigVMEditor/Private/Editor/RigVMEditor.cpp
+++ b/Engine/Plugins/Runtime/RigVM/Source/RigVMEditor/Private/Editor/RigVMEditor.cpp
@@ -968,7 +968,7 @@ void FRigVMEditorBase::PasteNodes()
 		TSharedPtr<SGraphEditor> GraphEditor = StaticCastSharedRef<SGraphEditor>(ActiveTab->GetContent());
 		if (GraphEditor.IsValid())
 		{
-			PasteLocation = GraphEditor->GetPasteLocation();
+			PasteLocation = GraphEditor->GetPasteLocation2f();
 
 		}
 	}
@@ -2729,7 +2729,7 @@ void FRigVMEditorBase::OnCreateComment()
 					{
 						Controller->OpenUndoBracket(TEXT("Create Comment"));
 						FEdGraphSchemaAction_K2AddComment CommentAction;
-						UEdGraphNode* EdNode = CommentAction.PerformAction(Graph, NULL, GraphEditor->GetPasteLocation());
+						UEdGraphNode* EdNode = CommentAction.PerformAction(Graph, NULL, GraphEditor->GetPasteLocation2f());
 						if (UEdGraphNode_Comment* CommentNode = CastChecked<UEdGraphNode_Comment>(EdNode))
 						{
 							Controller->SetNodeColorByName(CommentNode->GetFName(), CommentNode->CommentColor, false);
@@ -4029,7 +4029,7 @@ void FRigVMEditorBase::OnGraphNodeDropToPerform(TSharedPtr<FDragDropOperation> I
 				if (URigVMBlueprint* TargetRigBlueprint = Cast<URigVMBlueprint>(FBlueprintEditorUtils::FindBlueprintForGraph(InGraph)))
 				{
 					FProperty* Property = TargetRigBlueprint->SkeletonGeneratedClass->FindPropertyByName(*ExplorerOp->GetElement().Name);
-					Schema->RequestVariableDropOnPanel(InGraph, Property, InNodePosition, InScreenPosition);
+					Schema->RequestVariableDropOnPanel(InGraph, Property, UE::Slate::CastToVector2f(InNodePosition), UE::Slate::CastToVector2f(InScreenPosition));
 				}
 			}
 			else if (ExplorerOp->GetElement().Type == ERigVMExplorerElementType::LocalVariable)
diff --git a/Engine/Plugins/Runtime/RigVM/Source/RigVMEditor/Private/Editor/RigVMNewEditor.cpp b/Engine/Plugins/Runtime/RigVM/Source/RigVMEditor/Private/Editor/RigVMNewEditor.cpp
--- a/Engine/Plugins/Runtime/RigVM/Source/RigVMEditor/Private/Editor/RigVMNewEditor.cpp
+++ b/Engine/Plugins/Runtime/RigVM/Source/RigVMEditor/Private/Editor/RigVMNewEditor.cpp
@@ -1324,7 +1324,7 @@ TSharedRef<SGraphEditor> FRigVMNewEditor::CreateGraphEditorWidget(TSharedRef<FTa
 		
 	//OnSetPinVisibility.AddSP(&Editor.Get(), &SGraphEditor::SetPinVisibility);
 
-	FVector2D ViewOffset = FVector2D::ZeroVector;
+	FVector2f ViewOffset = FVector2f::ZeroVector;
 	float ZoomAmount = INDEX_NONE;
 
 	TSharedPtr<SDockTab> ActiveTab = DocumentManager->GetActiveTab();
diff --git a/Engine/Plugins/Runtime/RigVM/Source/RigVMEditor/Public/Editor/RigVMActionMenuItem.h b/Engine/Plugins/Runtime/RigVM/Source/RigVMEditor/Public/Editor/RigVMActionMenuItem.h
--- a/Engine/Plugins/Runtime/RigVM/Source/RigVMEditor/Public/Editor/RigVMActionMenuItem.h
+++ b/Engine/Plugins/Runtime/RigVM/Source/RigVMEditor/Public/Editor/RigVMActionMenuItem.h
@@ -46,8 +46,8 @@ public:
 	
 	// FEdGraphSchemaAction interface
 	virtual FName         GetTypeId() const final { return StaticGetTypeId(); }
-	virtual UEdGraphNode* PerformAction(UEdGraph* ParentGraph, UEdGraphPin* FromPin, FVector2D const Location, bool bSelectNewNode = true) final;
-	virtual UEdGraphNode* PerformAction(UEdGraph* ParentGraph, TArray<UEdGraphPin*>& FromPins, FVector2D const Location, bool bSelectNewNode = true) final;
+	virtual UEdGraphNode* PerformAction(UEdGraph* ParentGraph, UEdGraphPin* FromPin, const FVector2f& Location, bool bSelectNewNode = true) final;
+	virtual UEdGraphNode* PerformAction(UEdGraph* ParentGraph, TArray<UEdGraphPin*>& FromPins, const FVector2f& Location, bool bSelectNewNode = true) final;
 	virtual void          AddReferencedObjects(FReferenceCollector& Collector) final;
 	// End FEdGraphSchemaAction interface
 
diff --git a/Engine/Plugins/Runtime/RigVM/Source/RigVMEditor/Public/Editor/SRigVMActionMenu.cpp b/Engine/Plugins/Runtime/RigVM/Source/RigVMEditor/Public/Editor/SRigVMActionMenu.cpp
--- a/Engine/Plugins/Runtime/RigVM/Source/RigVMEditor/Public/Editor/SRigVMActionMenu.cpp
+++ b/Engine/Plugins/Runtime/RigVM/Source/RigVMEditor/Public/Editor/SRigVMActionMenu.cpp
@@ -36,7 +36,7 @@ struct FRigVMAction_PromoteVariable : public FEdGraphSchemaAction
 	}
 
 	// FEdGraphSchemaAction interface
-	virtual UEdGraphNode* PerformAction( class UEdGraph* ParentGraph, UEdGraphPin* FromPin, const FVector2D Location, bool bSelectNewNode = true) override
+	virtual UEdGraphNode* PerformAction( class UEdGraph* ParentGraph, UEdGraphPin* FromPin, const FVector2f& Location, bool bSelectNewNode = true) override
 	{
 		if( ( ParentGraph != NULL ) && ( FromPin != NULL ) )
 		{
@@ -49,7 +49,7 @@ struct FRigVMAction_PromoteVariable : public FEdGraphSchemaAction
 					{
 						if(const URigVMPin* ModelPin = RigNode->GetModelPinFromPinPath(FromPin->GetName()))
 						{
-							Controller->PromotePinToVariable(ModelPin->GetPinPath(), true, Location, true, true);
+							Controller->PromotePinToVariable(ModelPin->GetPinPath(), true, FDeprecateSlateVector2D(Location), true, true);
 						}
 					}
 				}
diff --git a/Engine/Plugins/Runtime/RigVM/Source/RigVMEditor/Public/Editor/SRigVMActionMenu.h b/Engine/Plugins/Runtime/RigVM/Source/RigVMEditor/Public/Editor/SRigVMActionMenu.h
--- a/Engine/Plugins/Runtime/RigVM/Source/RigVMEditor/Public/Editor/SRigVMActionMenu.h
+++ b/Engine/Plugins/Runtime/RigVM/Source/RigVMEditor/Public/Editor/SRigVMActionMenu.h
@@ -61,7 +61,7 @@ protected:
 private:
 	TObjectPtr<UEdGraph> GraphObj = nullptr;
 	TArray<UEdGraphPin*> DraggedFromPins;
-	FVector2D NewNodePosition;
+	FDeprecateSlateVector2D NewNodePosition;
 
 	TSharedPtr<SGraphActionMenu> GraphActionMenu;
 	TWeakPtr<IRigVMEditor> EditorPtr;
diff --git a/Engine/Plugins/Runtime/nDisplay/Source/DisplayClusterConfigurator/Private/Views/OutputMapping/DisplayClusterConfiguratorGraphSchema.cpp b/Engine/Plugins/Runtime/nDisplay/Source/DisplayClusterConfigurator/Private/Views/OutputMapping/DisplayClusterConfiguratorGraphSchema.cpp
--- a/Engine/Plugins/Runtime/nDisplay/Source/DisplayClusterConfigurator/Private/Views/OutputMapping/DisplayClusterConfiguratorGraphSchema.cpp
+++ b/Engine/Plugins/Runtime/nDisplay/Source/DisplayClusterConfigurator/Private/Views/OutputMapping/DisplayClusterConfiguratorGraphSchema.cpp
@@ -15,16 +15,16 @@
 FDisplayClusterConfiguratorSchemaAction_NewNode::FDisplayClusterConfiguratorSchemaAction_NewNode() :
 	FEdGraphSchemaAction(),
 	ItemType(EClusterItemType::ClusterNode),
-	PresetSize(FVector2D::ZeroVector)
+	PresetSize(FVector2f::ZeroVector)
 { }
 
-FDisplayClusterConfiguratorSchemaAction_NewNode::FDisplayClusterConfiguratorSchemaAction_NewNode(EClusterItemType InItemType, FVector2D InPresetSize, FText InDescription, FText InTooltip) :
+FDisplayClusterConfiguratorSchemaAction_NewNode::FDisplayClusterConfiguratorSchemaAction_NewNode(EClusterItemType InItemType, const UE::Slate::FDeprecateVector2DParameter& InPresetSize, FText InDescription, FText InTooltip) :
 	FEdGraphSchemaAction(LOCTEXT("NewNodeActionCategory", "nDisplay"), MoveTemp(InDescription), MoveTemp(InTooltip), 0),
 	ItemType(InItemType),
 	PresetSize(InPresetSize)
 { }
 
-UEdGraphNode* FDisplayClusterConfiguratorSchemaAction_NewNode::PerformAction(UEdGraph* ParentGraph, UEdGraphPin* FromPin, const FVector2D Location, bool bSelectNewNode)
+UEdGraphNode* FDisplayClusterConfiguratorSchemaAction_NewNode::PerformAction(UEdGraph* ParentGraph, UEdGraphPin* FromPin, const FVector2f& Location, bool bSelectNewNode)
 {
 	const UDisplayClusterConfiguratorGraph* ClusterGraph = CastChecked<UDisplayClusterConfiguratorGraph>(ParentGraph);
 	TSharedPtr<FDisplayClusterConfiguratorBlueprintEditor> Toolkit = ClusterGraph->GetToolkit().Pin();
@@ -58,7 +58,7 @@ UEdGraphNode* FDisplayClusterConfiguratorSchemaAction_NewNode::PerformAction(UEd
 
 void UDisplayClusterConfiguratorGraphSchema::GetGraphContextActions(FGraphContextMenuBuilder& ContextMenuBuilder) const
 {
-	const FVector2D PresetSize = FDisplayClusterConfiguratorPresetSize::CommonPresets[FDisplayClusterConfiguratorPresetSize::DefaultPreset].Size;
+	const FVector2f PresetSize = UE::Slate::CastToVector2f(FDisplayClusterConfiguratorPresetSize::CommonPresets[FDisplayClusterConfiguratorPresetSize::DefaultPreset].Size);
 
 	// Action for adding a new cluster node
 	{
diff --git a/Engine/Plugins/Runtime/nDisplay/Source/DisplayClusterConfigurator/Private/Views/OutputMapping/DisplayClusterConfiguratorGraphSchema.h b/Engine/Plugins/Runtime/nDisplay/Source/DisplayClusterConfigurator/Private/Views/OutputMapping/DisplayClusterConfiguratorGraphSchema.h
--- a/Engine/Plugins/Runtime/nDisplay/Source/DisplayClusterConfigurator/Private/Views/OutputMapping/DisplayClusterConfiguratorGraphSchema.h
+++ b/Engine/Plugins/Runtime/nDisplay/Source/DisplayClusterConfigurator/Private/Views/OutputMapping/DisplayClusterConfiguratorGraphSchema.h
@@ -25,10 +25,10 @@ struct FDisplayClusterConfiguratorSchemaAction_NewNode : public FEdGraphSchemaAc
 
 public:
 	FDisplayClusterConfiguratorSchemaAction_NewNode();
-	FDisplayClusterConfiguratorSchemaAction_NewNode(EClusterItemType InItemType, FVector2D InPresetSize, FText InDescription, FText InTooltip);
+	FDisplayClusterConfiguratorSchemaAction_NewNode(EClusterItemType InItemType, const UE::Slate::FDeprecateVector2DParameter& InPresetSize, FText InDescription, FText InTooltip);
 
 	//~ Begin FEdGraphSchemaAction Interface
-	virtual UEdGraphNode* PerformAction(UEdGraph* ParentGraph, UEdGraphPin* FromPin, const FVector2D Location, bool bSelectNewNode = true) override;
+	virtual UEdGraphNode* PerformAction(UEdGraph* ParentGraph, UEdGraphPin* FromPin, const FVector2f& Location, bool bSelectNewNode = true) override;
 	//~ End FEdGraphSchemaAction Interface
 
 private:
@@ -36,7 +36,7 @@ private:
 	EClusterItemType ItemType;
 
 	UPROPERTY()
-	FVector2D PresetSize;
+	FDeprecateSlateVector2D PresetSize;
 };
 
 UCLASS()
diff --git a/Engine/Plugins/Runtime/nDisplay/Source/DisplayClusterConfigurator/Private/Views/OutputMapping/SDisplayClusterConfiguratorGraphEditor.cpp b/Engine/Plugins/Runtime/nDisplay/Source/DisplayClusterConfigurator/Private/Views/OutputMapping/SDisplayClusterConfiguratorGraphEditor.cpp
--- a/Engine/Plugins/Runtime/nDisplay/Source/DisplayClusterConfigurator/Private/Views/OutputMapping/SDisplayClusterConfiguratorGraphEditor.cpp
+++ b/Engine/Plugins/Runtime/nDisplay/Source/DisplayClusterConfigurator/Private/Views/OutputMapping/SDisplayClusterConfiguratorGraphEditor.cpp
@@ -23,6 +23,7 @@
 #include "GraphEditorActions.h"
 #include "SGraphPanel.h"
 #include "Misc/ConfigCacheIni.h"
+#include "Misc/Optional.h"
 #include "IDocumentation.h"
 
 #include "Views/OutputMapping/IDisplayClusterConfiguratorViewOutputMapping.h"
@@ -620,7 +621,8 @@ void SDisplayClusterConfiguratorGraphEditor::PasteNodes()
 	{
 		FScopedTransaction Transaction(FText::Format(LOCTEXT("PasteClusterItems", "Paste Cluster {0}|plural(one=Item, other=Items)"), NumClusterItems));
 
-		TArray<UObject*> PastedObjects = UE::DisplayClusterConfiguratorClusterEditorUtils::PasteClusterItemsFromClipboard(TargetObjects, GetPasteLocation() / ViewScale);
+		FVector2D PasteLocation = GetPasteLocation2f() / ViewScale;
+		TArray<UObject*> PastedObjects = UE::DisplayClusterConfiguratorClusterEditorUtils::PasteClusterItemsFromClipboard(TargetObjects, PasteLocation);
 		if (PastedObjects.Num() > 0)
 		{
 			Toolkit->GetEditorData()->MarkPackageDirty();
diff --git a/Engine/Plugins/Runtime/nDisplay/Source/DisplayClusterConfigurator/Private/Views/OutputMapping/SDisplayClusterConfiguratorViewOutputMapping.cpp b/Engine/Plugins/Runtime/nDisplay/Source/DisplayClusterConfigurator/Private/Views/OutputMapping/SDisplayClusterConfiguratorViewOutputMapping.cpp
--- a/Engine/Plugins/Runtime/nDisplay/Source/DisplayClusterConfigurator/Private/Views/OutputMapping/SDisplayClusterConfiguratorViewOutputMapping.cpp
+++ b/Engine/Plugins/Runtime/nDisplay/Source/DisplayClusterConfigurator/Private/Views/OutputMapping/SDisplayClusterConfiguratorViewOutputMapping.cpp
@@ -101,7 +101,7 @@ void SDisplayClusterConfiguratorViewOutputMapping::Construct(const FArguments& I
 		InToolkit);
 
 	// Set initial zoom
-	FVector2D GraphEditorLocation = FVector2D::ZeroVector;
+	FVector2f GraphEditorLocation = FVector2f::ZeroVector;
 	float GraphEditorZoomAmount = 0;
 	InGraphEditor->GetViewLocation(GraphEditorLocation, GraphEditorZoomAmount);
 	InGraphEditor->SetViewLocation(GraphEditorLocation, 0.3f);
@@ -112,7 +112,7 @@ void SDisplayClusterConfiguratorViewOutputMapping::Tick(const FGeometry& Allotte
 	TSharedPtr<SDisplayClusterConfiguratorGraphEditor> GraphEditor = GraphEditorPtr.Pin();
 	check(GraphEditor.IsValid());
 
-	FVector2D GraphEditorLocation = FVector2D::ZeroVector;
+	FVector2f GraphEditorLocation = FVector2f::ZeroVector;
 	float GraphEditorZoomAmount = 0;
 	GraphEditor->GetViewLocation(GraphEditorLocation, GraphEditorZoomAmount);
 
@@ -122,11 +122,11 @@ void SDisplayClusterConfiguratorViewOutputMapping::Tick(const FGeometry& Allotte
 
 	// Compute the origin in absolute space.
 	FGeometry RootGeometry = PreviewSurface->GetCachedGeometry();
-	FVector2D AbsoluteOrigin = MakeGeometryWindowLocal(RootGeometry).LocalToAbsolute(FVector2D::ZeroVector);
-	FVector2D AbsoluteOriginWithOffset = AbsoluteOrigin - GraphEditorLocation * GraphEditorZoomAmount * DPIScale;
+	FVector2f AbsoluteOrigin = MakeGeometryWindowLocal(RootGeometry).LocalToAbsolute(FVector2f::ZeroVector);
+	FVector2f AbsoluteOriginWithOffset = AbsoluteOrigin - GraphEditorLocation * GraphEditorZoomAmount * DPIScale;
 
-	TopRuler->SetRuling(AbsoluteOriginWithOffset, 1.f / RulerScale);
-	SideRuler->SetRuling(AbsoluteOriginWithOffset, 1.f / RulerScale);
+	TopRuler->SetRuling(FDeprecateSlateVector2D(AbsoluteOriginWithOffset), 1.f / RulerScale);
+	SideRuler->SetRuling(FDeprecateSlateVector2D(AbsoluteOriginWithOffset), 1.f / RulerScale);
 
 	if (IsHovered())
 	{
@@ -250,13 +250,13 @@ FText SDisplayClusterConfiguratorViewOutputMapping::GetCursorPositionText() cons
 	const float ViewScale = ViewOutputMapping->GetOutputMappingSettings().ViewScale;
 
 	FGeometry RootGeometry = PreviewSurface->GetCachedGeometry();
-	const FVector2D CursorPos = RootGeometry.AbsoluteToLocal(FSlateApplication::Get().GetCursorPos());
+	const FVector2f CursorPos = RootGeometry.AbsoluteToLocal(FSlateApplication::Get().GetCursorPos());
 
-	FVector2D GraphEditorLocation = FVector2D::ZeroVector;
+	FVector2f GraphEditorLocation = FVector2f::ZeroVector;
 	float GraphEditorZoomAmount = 0;
 	GraphEditor->GetViewLocation(GraphEditorLocation, GraphEditorZoomAmount);
 
-	FVector2D GraphPosition = ((CursorPos / GraphEditorZoomAmount) + GraphEditorLocation) / ViewScale;
+	FVector2f GraphPosition = ((CursorPos / GraphEditorZoomAmount) + GraphEditorLocation) / ViewScale;
 
 	return FText::Format(LOCTEXT("CursorPositionFormat", "{0} x {1}"), FText::AsNumber(FMath::RoundToInt(GraphPosition.X)), FText::AsNumber(FMath::RoundToInt(GraphPosition.Y)));
 }
diff --git a/Engine/Plugins/Tests/EditorTests/Source/EditorTests/Private/UnrealEd/EditorBuildPromotionTests.cpp b/Engine/Plugins/Tests/EditorTests/Source/EditorTests/Private/UnrealEd/EditorBuildPromotionTests.cpp
--- a/Engine/Plugins/Tests/EditorTests/Source/EditorTests/Private/UnrealEd/EditorBuildPromotionTests.cpp
+++ b/Engine/Plugins/Tests/EditorTests/Source/EditorTests/Private/UnrealEd/EditorBuildPromotionTests.cpp
@@ -250,7 +250,7 @@ namespace EditorBuildPromotionTestUtils
 	* @param GraphLocation - The location to place the node
 	* @param ConnectPin - The pin to connect the node to
 	*/
-	static UEdGraphNode* CreateNewGraphNodeFromTemplate(UK2Node* NodeTemplate, UEdGraph* InGraph, const FVector2D& GraphLocation, UEdGraphPin* ConnectPin = NULL)
+	static UEdGraphNode* CreateNewGraphNodeFromTemplate(UK2Node* NodeTemplate, UEdGraph* InGraph, const FVector2f& GraphLocation, UEdGraphPin* ConnectPin = NULL)
 	{
 		TSharedPtr<FEdGraphSchemaAction_K2NewNode> Action = TSharedPtr<FEdGraphSchemaAction_K2NewNode>(new FEdGraphSchemaAction_K2NewNode(FText::GetEmpty(), FText::GetEmpty(), FText::GetEmpty(), 0));
 		Action->NodeTemplate = NodeTemplate;
@@ -279,7 +279,7 @@ namespace EditorBuildPromotionTestUtils
 
 		if (!ExistingEvent)
 		{
-			return CreateNewGraphNodeFromTemplate(NewEventNode, InGraph, FVector2D(200, 0));
+			return CreateNewGraphNodeFromTemplate(NewEventNode, InGraph, FVector2f(200.0f, 0.0f));
 		}
 		return ExistingEvent;
 	}
@@ -300,7 +300,7 @@ namespace EditorBuildPromotionTestUtils
 		UK2Node_CustomEvent* NewEventNode = NewObject<UK2Node_CustomEvent>(TempOuter);
 		NewEventNode->CustomFunctionName = "EventName";
 
-		return CreateNewGraphNodeFromTemplate(NewEventNode, InGraph, FVector2D(1200, 0));
+		return CreateNewGraphNodeFromTemplate(NewEventNode, InGraph, FVector2f(1200.0f, 0.0f));
 	}
 
 	/**
@@ -334,7 +334,7 @@ namespace EditorBuildPromotionTestUtils
 		static FName PrintStringFunctionName(TEXT("PrintString"));
 		UK2Node* CallFuncNode = CreateKismetFunctionTemplate(TempOuter, PrintStringFunctionName);
 
-		return CreateNewGraphNodeFromTemplate(CallFuncNode, InGraph, FVector2D(680, 0), ConnectPin);
+		return CreateNewGraphNodeFromTemplate(CallFuncNode, InGraph, FVector2f(680.0f, 0.0f), ConnectPin);
 	}
 
 	/**
diff --git a/Engine/Plugins/TextureGraph/Source/TextureGraphEditor/Private/EdGraph/TG_EdGraphSchema.cpp b/Engine/Plugins/TextureGraph/Source/TextureGraphEditor/Private/EdGraph/TG_EdGraphSchema.cpp
--- a/Engine/Plugins/TextureGraph/Source/TextureGraphEditor/Private/EdGraph/TG_EdGraphSchema.cpp
+++ b/Engine/Plugins/TextureGraph/Source/TextureGraphEditor/Private/EdGraph/TG_EdGraphSchema.cpp
@@ -378,13 +378,13 @@ void UTG_EdGraphSchema::BreakSinglePinLink(UEdGraphPin* SourcePin, UEdGraphPin*
 	TextureGraph->RemoveEdge(*SourceTSNode, SourcePin->PinName, *TargetTSNode, TargetPin->PinName);
 }
 
-void UTG_EdGraphSchema::DroppedAssetsOnGraph(const TArray<FAssetData>& Assets, const FVector2D& GraphPosition, UEdGraph* Graph) const
+void UTG_EdGraphSchema::DroppedAssetsOnGraph(const TArray<FAssetData>& Assets, const FVector2f& GraphPosition, UEdGraph* Graph) const
 {
 	UTG_EdGraph* TG_Graph = CastChecked<UTG_EdGraph>(Graph);
 
 	const int32 LocOffsetBetweenNodes = 32;
 
-	FVector2D ExpressionPosition = GraphPosition;
+	FVector2f ExpressionPosition = GraphPosition;
 
 	for (int32 AssetIdx = 0; AssetIdx < Assets.Num(); ++AssetIdx)
 	{
@@ -436,7 +436,7 @@ void UTG_EdGraphSchema::DroppedAssetsOnGraph(const TArray<FAssetData>& Assets, c
 	}
 }
 
-void UTG_EdGraphSchema::DroppedAssetsOnNode(const TArray<FAssetData>& Assets, const FVector2D& GraphPosition, UEdGraphNode* Node) const
+void UTG_EdGraphSchema::DroppedAssetsOnNode(const TArray<FAssetData>& Assets, const FVector2f& GraphPosition, UEdGraphNode* Node) const
 {
 	const UTG_EdGraphNode* EdNode = Cast<UTG_EdGraphNode>(Node);
 
diff --git a/Engine/Plugins/TextureGraph/Source/TextureGraphEditor/Private/EdGraph/TG_EdGraphSchema.h b/Engine/Plugins/TextureGraph/Source/TextureGraphEditor/Private/EdGraph/TG_EdGraphSchema.h
--- a/Engine/Plugins/TextureGraph/Source/TextureGraphEditor/Private/EdGraph/TG_EdGraphSchema.h
+++ b/Engine/Plugins/TextureGraph/Source/TextureGraphEditor/Private/EdGraph/TG_EdGraphSchema.h
@@ -69,8 +69,8 @@ public:
 	virtual void BreakPinLinks(UEdGraphPin& TargetPin, bool bSendsNodeNotifcation) const override;
 	virtual void BreakSinglePinLink(UEdGraphPin* SourcePin, UEdGraphPin* TargetPin) const override;
 
-	virtual void DroppedAssetsOnGraph(const TArray<struct FAssetData>& Assets, const FVector2D& GraphPosition, UEdGraph* Graph) const override;
-	virtual void DroppedAssetsOnNode(const TArray<struct FAssetData>& Assets, const FVector2D& GraphPosition, UEdGraphNode* Node) const override;
+	virtual void DroppedAssetsOnGraph(const TArray<struct FAssetData>& Assets, const FVector2f& GraphPosition, UEdGraph* Graph) const override;
+	virtual void DroppedAssetsOnNode(const TArray<struct FAssetData>& Assets, const FVector2f& GraphPosition, UEdGraphNode* Node) const override;
 	virtual void GetAssetsGraphHoverMessage(const TArray<FAssetData>& Assets, const UEdGraph* HoverGraph, FString& OutTooltipText, bool& OutOkIcon) const override;
 	virtual void GetAssetsNodeHoverMessage(const TArray<struct FAssetData>& Assets, const UEdGraphNode* HoverNode, FString& OutTooltipText, bool& OutOkIcon) const override;
 	virtual void TrySetDefaultValue(UEdGraphPin& Pin, const FString& NewDefaultValue, bool bMarkAsModified = true) const override;
diff --git a/Engine/Plugins/TextureGraph/Source/TextureGraphEditor/Private/EdGraph/TG_EdGraphSchemaActions.cpp b/Engine/Plugins/TextureGraph/Source/TextureGraphEditor/Private/EdGraph/TG_EdGraphSchemaActions.cpp
--- a/Engine/Plugins/TextureGraph/Source/TextureGraphEditor/Private/EdGraph/TG_EdGraphSchemaActions.cpp
+++ b/Engine/Plugins/TextureGraph/Source/TextureGraphEditor/Private/EdGraph/TG_EdGraphSchemaActions.cpp
@@ -13,12 +13,12 @@
 
 #define LOCTEXT_NAMESPACE "TG_EdGraphSchemaActions"
 
-UEdGraphNode* FTG_EdGraphSchemaAction_NewNode::PerformAction(UEdGraph* ParentGraph, UEdGraphPin* FromPin, const FVector2D Location, bool bSelectNewNode)
+UEdGraphNode* FTG_EdGraphSchemaAction_NewNode::PerformAction(UEdGraph* ParentGraph, UEdGraphPin* FromPin, const FVector2f& Location, bool bSelectNewNode)
 {
 	return CreateExpressionNode(ParentGraph, TG_ExpressionClass.Get(), FromPin, Location, bSelectNewNode);
 }
 
-UEdGraphNode* FTG_EdGraphSchemaAction_NewNode::CreateExpressionNode(UEdGraph* ParentGraph, const UClass* ExpressionClass, UEdGraphPin* FromPin, const FVector2D Location, bool bSelectNewNode)
+UEdGraphNode* FTG_EdGraphSchemaAction_NewNode::CreateExpressionNode(UEdGraph* ParentGraph, const UClass* ExpressionClass, UEdGraphPin* FromPin, const UE::Slate::FDeprecateVector2DParameter& Location, bool bSelectNewNode)
 {
 	const FScopedTransaction Transaction(NSLOCTEXT("UnrealEd", "TSEditorNewNode", "TS Editor: New Node"));
 
@@ -32,7 +32,7 @@ UEdGraphNode* FTG_EdGraphSchemaAction_NewNode::CreateExpressionNode(UEdGraph* Pa
 	
 	if (NewExpressionNode)
 	{	
-		UTG_EdGraphNode* NewNode = TG_EdGraph->AddModelNode(NewExpressionNode, true, Location);
+		UTG_EdGraphNode* NewNode = TG_EdGraph->AddModelNode(NewExpressionNode, true, FDeprecateSlateVector2D(Location));
 		// if (MaterialExpressionClass == UMaterialExpressionFunctionInput::StaticClass() && FromPin)
 		// {
 		// 	// Set this to be an input of the type we dragged from
@@ -47,7 +47,7 @@ UEdGraphNode* FTG_EdGraphSchemaAction_NewNode::CreateExpressionNode(UEdGraph* Pa
 	return nullptr;
 }
 
-UEdGraphNode* FTG_EdGraphSchemaAction_NewNode::CreateExpressionNode(UEdGraph* ParentGraph, UTG_Expression* Expression, UEdGraphPin* FromPin,  const FVector2D Location, bool bSelectNewNode)
+UEdGraphNode* FTG_EdGraphSchemaAction_NewNode::CreateExpressionNode(UEdGraph* ParentGraph, UTG_Expression* Expression, UEdGraphPin* FromPin,  const UE::Slate::FDeprecateVector2DParameter& Location, bool bSelectNewNode)
 {
 	const FScopedTransaction Transaction(NSLOCTEXT("UnrealEd", "TSEditorNewNode", "TS Editor: New Node"));
 
@@ -61,7 +61,7 @@ UEdGraphNode* FTG_EdGraphSchemaAction_NewNode::CreateExpressionNode(UEdGraph* Pa
 	
 	if (NewExpressionNode)
 	{	
-		UTG_EdGraphNode* NewNode = TG_EdGraph->AddModelNode(NewExpressionNode, true, Location);
+		UTG_EdGraphNode* NewNode = TG_EdGraph->AddModelNode(NewExpressionNode, true, FDeprecateSlateVector2D(Location));
 		NewNode->AutowireNewNode(FromPin);
 		return NewNode;
 	}
@@ -70,13 +70,13 @@ UEdGraphNode* FTG_EdGraphSchemaAction_NewNode::CreateExpressionNode(UEdGraph* Pa
 }
 
 UEdGraphNode* FTG_EdGraphSchemaAction_NewComment::PerformAction(UEdGraph* ParentGraph, UEdGraphPin* FromPin,
-	const FVector2D Location, bool bSelectNewNode)
+	const FVector2f& Location, bool bSelectNewNode)
 {
 	UEdGraphNode_Comment* CommentTemplate = NewObject<UEdGraphNode_Comment>();
 
 	TSharedPtr<SGraphEditor> GraphEditorPtr = SGraphEditor::FindGraphEditorForGraph(ParentGraph);
 
-	FVector2D SpawnLocation = Location;
+	FVector2f SpawnLocation = Location;
 	FSlateRect Bounds;
 	if (GraphEditorPtr && GraphEditorPtr->GetBoundsForSelectedNodes(Bounds, 50.0f))
 	{
diff --git a/Engine/Plugins/TextureGraph/Source/TextureGraphEditor/Private/EdGraph/TG_EdGraphSchemaActions.h b/Engine/Plugins/TextureGraph/Source/TextureGraphEditor/Private/EdGraph/TG_EdGraphSchemaActions.h
--- a/Engine/Plugins/TextureGraph/Source/TextureGraphEditor/Private/EdGraph/TG_EdGraphSchemaActions.h
+++ b/Engine/Plugins/TextureGraph/Source/TextureGraphEditor/Private/EdGraph/TG_EdGraphSchemaActions.h
@@ -25,11 +25,12 @@ public:
 	TSubclassOf<UTG_Expression> TG_ExpressionClass;
 
 	//~ Begin FEdGraphSchemaAction Interface
-	virtual UEdGraphNode* PerformAction(UEdGraph* ParentGraph, UEdGraphPin* FromPin, const FVector2D Location, bool bSelectNewNode = true) override;
+	using FEdGraphSchemaAction::PerformAction; // Prevent hiding of FVector2D deprecated function in base class
+	virtual UEdGraphNode* PerformAction(UEdGraph* ParentGraph, UEdGraphPin* FromPin, const FVector2f& Location, bool bSelectNewNode = true) override;
 	//~ End FEdGraphSchemaAction Interface
 
-	static UEdGraphNode* CreateExpressionNode(UEdGraph* ParentGraph, const UClass* ExpressionClass, UEdGraphPin* FromPin,  const FVector2D Location, bool bSelectNewNode = true);
-	static UEdGraphNode* CreateExpressionNode(UEdGraph* ParentGraph, UTG_Expression* Expression, UEdGraphPin* FromPin, const FVector2D Location, bool bSelectNewNode = true);
+	static UEdGraphNode* CreateExpressionNode(UEdGraph* ParentGraph, const UClass* ExpressionClass, UEdGraphPin* FromPin,  const UE::Slate::FDeprecateVector2DParameter& Location, bool bSelectNewNode = true);
+	static UEdGraphNode* CreateExpressionNode(UEdGraph* ParentGraph, UTG_Expression* Expression, UEdGraphPin* FromPin, const UE::Slate::FDeprecateVector2DParameter& Location, bool bSelectNewNode = true);
 };
 
 /** Action to add a 'comment' node to the graph */
@@ -50,7 +51,8 @@ struct FTG_EdGraphSchemaAction_NewComment : public FEdGraphSchemaAction
 
 	// FEdGraphSchemaAction interface
 	virtual FName GetTypeId() const override { return StaticGetTypeId(); }
-	virtual UEdGraphNode* PerformAction(class UEdGraph* ParentGraph, UEdGraphPin* FromPin, const FVector2D Location, bool bSelectNewNode = true) override;
+	using FEdGraphSchemaAction::PerformAction; // Prevent hiding of FVector2D deprecated function in base class
+	virtual UEdGraphNode* PerformAction(class UEdGraph* ParentGraph, UEdGraphPin* FromPin, const FVector2f& Location, bool bSelectNewNode = true) override;
 	// End of FEdGraphSchemaAction interface
 };
 
diff --git a/Engine/Plugins/TextureGraph/Source/TextureGraphEditor/Private/TG_Editor.cpp b/Engine/Plugins/TextureGraph/Source/TextureGraphEditor/Private/TG_Editor.cpp
--- a/Engine/Plugins/TextureGraph/Source/TextureGraphEditor/Private/TG_Editor.cpp
+++ b/Engine/Plugins/TextureGraph/Source/TextureGraphEditor/Private/TG_Editor.cpp
@@ -762,10 +762,10 @@ void FTG_Editor::OnCreateComment()
 		FTG_EdGraphSchemaAction_NewComment CommentAction;
 
 		TSharedPtr<SGraphEditor> GraphEditorPtr = SGraphEditor::FindGraphEditorForGraph(TG_EdGraph);
-		FVector2D Location = FVector2D::ZeroVector;
+		FVector2f Location = FVector2f::ZeroVector;
 		if (GraphEditorPtr)
 		{
-			Location = GraphEditorPtr->GetPasteLocation();
+			Location = GraphEditorPtr->GetPasteLocation2f();
 		}
 
 		CommentAction.PerformAction(TG_EdGraph, nullptr, Location);
@@ -1370,7 +1370,7 @@ void FTG_Editor::PasteNodes()
 {
 	if (TSharedPtr<SGraphEditor> FocusedGraphEd = GraphEditorWidget)
 	{
-		PasteNodesHere(FocusedGraphEd->GetPasteLocation(), FocusedGraphEd->GetCurrentGraph());
+		PasteNodesHere(FocusedGraphEd->GetPasteLocation2f(), FocusedGraphEd->GetCurrentGraph());
 	}
 }
 
diff --git a/Engine/Plugins/TextureGraph/Source/TextureGraphEditor/Private/Widgets/ActionMenuWidgets/STG_GraphActionMenu.h b/Engine/Plugins/TextureGraph/Source/TextureGraphEditor/Private/Widgets/ActionMenuWidgets/STG_GraphActionMenu.h
--- a/Engine/Plugins/TextureGraph/Source/TextureGraphEditor/Private/Widgets/ActionMenuWidgets/STG_GraphActionMenu.h
+++ b/Engine/Plugins/TextureGraph/Source/TextureGraphEditor/Private/Widgets/ActionMenuWidgets/STG_GraphActionMenu.h
@@ -64,7 +64,7 @@ protected:
 
 	TArray<UEdGraphPin*> DraggedFromPins;
 
-	FVector2D NewNodePosition;
+	FDeprecateSlateVector2D NewNodePosition;
 
 	bool AutoExpandActionMenu;
 
diff --git a/Engine/Source/Editor/AIGraph/Classes/AIGraphSchema.h b/Engine/Source/Editor/AIGraph/Classes/AIGraphSchema.h
--- a/Engine/Source/Editor/AIGraph/Classes/AIGraphSchema.h
+++ b/Engine/Source/Editor/AIGraph/Classes/AIGraphSchema.h
@@ -24,7 +24,7 @@ struct AIGRAPH_API FAISchemaAction_AddComment : public FEdGraphSchemaAction
 	}
 
 	// FEdGraphSchemaAction interface
-	virtual UEdGraphNode* PerformAction(class UEdGraph* ParentGraph, UEdGraphPin* FromPin, const FVector2D Location, bool bSelectNewNode = true) override final;
+	virtual UEdGraphNode* PerformAction(class UEdGraph* ParentGraph, UEdGraphPin* FromPin, const FVector2f& Location, bool bSelectNewNode = true) override final;
 	// End of FEdGraphSchemaAction interface
 };
 
@@ -49,19 +49,10 @@ struct AIGRAPH_API FAISchemaAction_NewNode : public FEdGraphSchemaAction
 	{}
 
 	//~ Begin FEdGraphSchemaAction Interface
-	virtual UEdGraphNode* PerformAction(class UEdGraph* ParentGraph, UEdGraphPin* FromPin, const FVector2D Location, bool bSelectNewNode = true) override;
-	virtual UEdGraphNode* PerformAction(class UEdGraph* ParentGraph, TArray<UEdGraphPin*>& FromPins, const FVector2D Location, bool bSelectNewNode = true) override;
+	virtual UEdGraphNode* PerformAction(class UEdGraph* ParentGraph, UEdGraphPin* FromPin, const FVector2f& Location, bool bSelectNewNode = true) override;
+	virtual UEdGraphNode* PerformAction(class UEdGraph* ParentGraph, TArray<UEdGraphPin*>& FromPins, const FVector2f& Location, bool bSelectNewNode = true) override;
 	virtual void AddReferencedObjects(FReferenceCollector& Collector) override;
 	//~ End FEdGraphSchemaAction Interface
-
-	template <typename NodeType>
-	static NodeType* SpawnNodeFromTemplate(class UEdGraph* ParentGraph, NodeType* InTemplateNode, const FVector2D Location)
-	{
-		FAISchemaAction_NewNode Action;
-		Action.NodeTemplate = InTemplateNode;
-
-		return Cast<NodeType>(Action.PerformAction(ParentGraph, nullptr, Location));
-	}
 };
 
 /** Action to add a subnode to the selected node */
@@ -91,8 +82,8 @@ struct AIGRAPH_API FAISchemaAction_NewSubNode : public FEdGraphSchemaAction
 	{}
 
 	//~ Begin FEdGraphSchemaAction Interface
-	virtual UEdGraphNode* PerformAction(class UEdGraph* ParentGraph, UEdGraphPin* FromPin, const FVector2D Location, bool bSelectNewNode = true) override;
-	virtual UEdGraphNode* PerformAction(class UEdGraph* ParentGraph, TArray<UEdGraphPin*>& FromPins, const FVector2D Location, bool bSelectNewNode = true) override;
+	virtual UEdGraphNode* PerformAction(class UEdGraph* ParentGraph, UEdGraphPin* FromPin, const FVector2f& Location, bool bSelectNewNode = true) override;
+	virtual UEdGraphNode* PerformAction(class UEdGraph* ParentGraph, TArray<UEdGraphPin*>& FromPins, const FVector2f& Location, bool bSelectNewNode = true) override;
 	virtual void AddReferencedObjects(FReferenceCollector& Collector) override;
 	//~ End FEdGraphSchemaAction Interface
 };
diff --git a/Engine/Source/Editor/AIGraph/Private/AIGraphEditor.cpp b/Engine/Source/Editor/AIGraph/Private/AIGraphEditor.cpp
--- a/Engine/Source/Editor/AIGraph/Private/AIGraphEditor.cpp
+++ b/Engine/Source/Editor/AIGraph/Private/AIGraphEditor.cpp
@@ -308,7 +308,7 @@ void FAIGraphEditor::PasteNodes()
 {
 	if (TSharedPtr<SGraphEditor> CurrentGraphEditor = UpdateGraphEdPtr.Pin())
 	{
-		PasteNodesHere(CurrentGraphEditor->GetPasteLocation());
+		PasteNodesHere(CurrentGraphEditor->GetPasteLocation2f());
 	}
 }
 
@@ -523,7 +523,7 @@ void FAIGraphEditor::OnCreateComment()
 		TSharedPtr<FEdGraphSchemaAction> Action = EdGraph->GetSchema()->GetCreateCommentAction();
 		if (Action.IsValid())
 		{
-			Action->PerformAction(EdGraph, nullptr, FVector2D());
+			Action->PerformAction(EdGraph, nullptr, FVector2f());
 		}
 	}
 }
diff --git a/Engine/Source/Editor/AIGraph/Private/AIGraphSchema.cpp b/Engine/Source/Editor/AIGraph/Private/AIGraphSchema.cpp
--- a/Engine/Source/Editor/AIGraph/Private/AIGraphSchema.cpp
+++ b/Engine/Source/Editor/AIGraph/Private/AIGraphSchema.cpp
@@ -24,11 +24,11 @@ namespace
 
 //////////////////////////////////////////////////////////////////////////
 
-UEdGraphNode* FAISchemaAction_AddComment::PerformAction(class UEdGraph* ParentGraph, UEdGraphPin* FromPin, const FVector2D Location, bool bSelectNewNode)
+UEdGraphNode* FAISchemaAction_AddComment::PerformAction(class UEdGraph* ParentGraph, UEdGraphPin* FromPin, const FVector2f& Location, bool bSelectNewNode)
 {
 	UEdGraphNode_Comment* const CommentTemplate = NewObject<UEdGraphNode_Comment>();
 
-	FVector2D SpawnLocation = Location;
+	FVector2f SpawnLocation = Location;
 	FSlateRect Bounds;
 
 	TSharedPtr<SGraphEditor> GraphEditorPtr = SGraphEditor::FindGraphEditorForGraph(ParentGraph);
@@ -44,7 +44,7 @@ UEdGraphNode* FAISchemaAction_AddComment::PerformAction(class UEdGraph* ParentGr
 		else
 		{
 			// Otherwise initialize a default comment at the user's cursor location.
-			SpawnLocation = GraphEditorPtr->GetPasteLocation();
+			SpawnLocation = GraphEditorPtr->GetPasteLocation2f();
 		}
 	}
 
@@ -55,7 +55,7 @@ UEdGraphNode* FAISchemaAction_AddComment::PerformAction(class UEdGraph* ParentGr
 
 //////////////////////////////////////////////////////////////////////////
 
-UEdGraphNode* FAISchemaAction_NewNode::PerformAction(class UEdGraph* ParentGraph, UEdGraphPin* FromPin, const FVector2D Location, bool bSelectNewNode)
+UEdGraphNode* FAISchemaAction_NewNode::PerformAction(class UEdGraph* ParentGraph, UEdGraphPin* FromPin, const FVector2f& Location, bool bSelectNewNode)
 {
 	UEdGraphNode* ResultNode = NULL;
 
@@ -84,7 +84,7 @@ UEdGraphNode* FAISchemaAction_NewNode::PerformAction(class UEdGraph* ParentGraph
 		if (FromPin && FromPin->Direction == EGPD_Input)
 		{
 			UEdGraphNode* PinNode = FromPin->GetOwningNode();
-			const FVector::FReal XDelta = FMath::Abs(PinNode->NodePosX - Location.X);
+			const float XDelta = FMath::Abs(PinNode->NodePosX - Location.X);
 
 			if (XDelta < NodeDistance)
 			{
@@ -108,7 +108,7 @@ UEdGraphNode* FAISchemaAction_NewNode::PerformAction(class UEdGraph* ParentGraph
 	return ResultNode;
 }
 
-UEdGraphNode* FAISchemaAction_NewNode::PerformAction(class UEdGraph* ParentGraph, TArray<UEdGraphPin*>& FromPins, const FVector2D Location, bool bSelectNewNode)
+UEdGraphNode* FAISchemaAction_NewNode::PerformAction(class UEdGraph* ParentGraph, TArray<UEdGraphPin*>& FromPins, const FVector2f& Location, bool bSelectNewNode)
 {
 	UEdGraphNode* ResultNode = NULL;
 	if (FromPins.Num() > 0)
@@ -137,13 +137,13 @@ void FAISchemaAction_NewNode::AddReferencedObjects(FReferenceCollector& Collecto
 	Collector.AddReferencedObject(NodeTemplate);
 }
 
-UEdGraphNode* FAISchemaAction_NewSubNode::PerformAction(class UEdGraph* ParentGraph, UEdGraphPin* FromPin, const FVector2D Location, bool bSelectNewNode)
+UEdGraphNode* FAISchemaAction_NewSubNode::PerformAction(class UEdGraph* ParentGraph, UEdGraphPin* FromPin, const FVector2f& Location, bool bSelectNewNode)
 {
 	ParentNode->AddSubNode(NodeTemplate, ParentGraph);
 	return NULL;
 }
 
-UEdGraphNode* FAISchemaAction_NewSubNode::PerformAction(class UEdGraph* ParentGraph, TArray<UEdGraphPin*>& FromPins, const FVector2D Location, bool bSelectNewNode)
+UEdGraphNode* FAISchemaAction_NewSubNode::PerformAction(class UEdGraph* ParentGraph, TArray<UEdGraphPin*>& FromPins, const FVector2f& Location, bool bSelectNewNode)
 {
 	return PerformAction(ParentGraph, NULL, Location, bSelectNewNode);
 }
diff --git a/Engine/Source/Editor/AIGraph/Private/SGraphEditorActionMenuAI.cpp b/Engine/Source/Editor/AIGraph/Private/SGraphEditorActionMenuAI.cpp
--- a/Engine/Source/Editor/AIGraph/Private/SGraphEditorActionMenuAI.cpp
+++ b/Engine/Source/Editor/AIGraph/Private/SGraphEditorActionMenuAI.cpp
@@ -94,7 +94,7 @@ void SGraphEditorActionMenuAI::OnActionSelected( const TArray< TSharedPtr<FEdGra
 
 				if ( CurrentAction.IsValid() )
 				{
-					CurrentAction->PerformAction(GraphObj, DraggedFromPins, NewNodePosition);
+					CurrentAction->PerformAction(GraphObj, DraggedFromPins, UE::Slate::CastToVector2f(NewNodePosition));
 					bDoDismissMenus = true;
 				}
 			}
diff --git a/Engine/Source/Editor/AnimGraph/Private/AnimationGraphSchema.cpp b/Engine/Source/Editor/AnimGraph/Private/AnimationGraphSchema.cpp
--- a/Engine/Source/Editor/AnimGraph/Private/AnimationGraphSchema.cpp
+++ b/Engine/Source/Editor/AnimGraph/Private/AnimationGraphSchema.cpp
@@ -608,7 +608,7 @@ bool IsAimOffsetBlendSpace(UBlendSpace* BlendSpace)
 			BlendSpace->IsA(UAimOffsetBlendSpace1D::StaticClass());
 }
 
-void UAnimationGraphSchema::SpawnNodeFromAsset(UAnimationAsset* Asset, const FVector2D& GraphPosition, UEdGraph* Graph, UEdGraphPin* PinIfAvailable)
+void UAnimationGraphSchema::SpawnNodeFromAsset(UAnimationAsset* Asset, const UE::Slate::FDeprecateVector2DParameter& GraphPosition, UEdGraph* Graph, UEdGraphPin* PinIfAvailable)
 {
 	check(Graph);
 	check(Graph->GetSchema()->IsA(UAnimationGraphSchema::StaticClass()));
@@ -640,7 +640,7 @@ void UAnimationGraphSchema::SpawnNodeFromAsset(UAnimationAsset* Asset, const FVe
 	}
 }
 
-void UAnimationGraphSchema::SpawnRigidBodyNodeFromAsset(UPhysicsAsset* Asset, const FVector2D& GraphPosition, UEdGraph* Graph)
+void UAnimationGraphSchema::SpawnRigidBodyNodeFromAsset(UPhysicsAsset* Asset, const UE::Slate::FDeprecateVector2DParameter& GraphPosition, UEdGraph* Graph)
 {
 	check(Graph);
 	check(Graph->GetSchema()->IsA(UAnimationGraphSchema::StaticClass()));
@@ -675,7 +675,7 @@ void UAnimationGraphSchema::UpdateNodeWithAsset(UK2Node* K2Node, UAnimationAsset
 	}
 }
 
-void UAnimationGraphSchema::DroppedAssetsOnGraph( const TArray<FAssetData>& Assets, const FVector2D& GraphPosition, UEdGraph* Graph ) const 
+void UAnimationGraphSchema::DroppedAssetsOnGraph( const TArray<FAssetData>& Assets, const FVector2f& GraphPosition, UEdGraph* Graph ) const 
 {
 	if (Graph != NULL)
 	{
@@ -685,12 +685,12 @@ void UAnimationGraphSchema::DroppedAssetsOnGraph( const TArray<FAssetData>& Asse
 		}
 		else if (UPhysicsAsset* PhysicsAsset = FAssetData::GetFirstAsset<UPhysicsAsset>(Assets))
 		{
-			SpawnRigidBodyNodeFromAsset(PhysicsAsset, GraphPosition, Graph);
+			SpawnRigidBodyNodeFromAsset(PhysicsAsset, FDeprecateSlateVector2D(GraphPosition), Graph);
 		}
 	}
 }
 
-void UAnimationGraphSchema::DroppedAssetsOnNode(const TArray<FAssetData>& Assets, const FVector2D& GraphPosition, UEdGraphNode* Node) const
+void UAnimationGraphSchema::DroppedAssetsOnNode(const TArray<FAssetData>& Assets, const FVector2f& GraphPosition, UEdGraphNode* Node) const
 {
 	UAnimationAsset* Asset = FAssetData::GetFirstAsset<UAnimationAsset>(Assets);
 	UK2Node* K2Node = Cast<UK2Node>(Node);
@@ -700,13 +700,13 @@ void UAnimationGraphSchema::DroppedAssetsOnNode(const TArray<FAssetData>& Assets
 	}
 }
 
-void UAnimationGraphSchema::DroppedAssetsOnPin(const TArray<FAssetData>& Assets, const FVector2D& GraphPosition, UEdGraphPin* Pin) const
+void UAnimationGraphSchema::DroppedAssetsOnPin(const TArray<FAssetData>& Assets, const FVector2f& GraphPosition, UEdGraphPin* Pin) const
 {
 	UAnimationAsset* Asset = FAssetData::GetFirstAsset<UAnimationAsset>(Assets);
 	if ((Asset != NULL) && (Pin != NULL))
 	{
 		// Don't have access to bounding information for node, using fixed offset that should work for most cases.
-		const FVector2D FixedOffset(-250.0f, 50.0f);
+		const FVector2f FixedOffset(-250.0f, 50.0f);
 		SpawnNodeFromAsset(Asset, GraphPosition + FixedOffset, Pin->GetOwningNode()->GetGraph(), Pin);
 	}
 }
diff --git a/Engine/Source/Editor/AnimGraph/Private/AnimationStateMachineSchema.cpp b/Engine/Source/Editor/AnimGraph/Private/AnimationStateMachineSchema.cpp
--- a/Engine/Source/Editor/AnimGraph/Private/AnimationStateMachineSchema.cpp
+++ b/Engine/Source/Editor/AnimGraph/Private/AnimationStateMachineSchema.cpp
@@ -44,7 +44,7 @@ TSharedPtr<FEdGraphSchemaAction_NewStateNode> AddNewStateNodeAction(FGraphContex
 /////////////////////////////////////////////////////
 // FEdGraphSchemaAction_NewStateNode
 
-UEdGraphNode* FEdGraphSchemaAction_NewStateNode::PerformAction(class UEdGraph* ParentGraph, UEdGraphPin* FromPin, const FVector2D Location, bool bSelectNewNode/* = true*/)
+UEdGraphNode* FEdGraphSchemaAction_NewStateNode::PerformAction(class UEdGraph* ParentGraph, UEdGraphPin* FromPin, const FVector2f& Location, bool bSelectNewNode/* = true*/)
 {
 	UEdGraphNode* ResultNode = NULL;
 
@@ -93,14 +93,14 @@ void FEdGraphSchemaAction_NewStateNode::AddReferencedObjects( FReferenceCollecto
 /////////////////////////////////////////////////////
 // FEdGraphSchemaAction_NewStateComment
 
-UEdGraphNode* FEdGraphSchemaAction_NewStateComment::PerformAction(class UEdGraph* ParentGraph, UEdGraphPin* FromPin, const FVector2D Location, bool bSelectNewNode/* = true*/)
+UEdGraphNode* FEdGraphSchemaAction_NewStateComment::PerformAction(class UEdGraph* ParentGraph, UEdGraphPin* FromPin, const FVector2f& Location, bool bSelectNewNode/* = true*/)
 {
 	// Add menu item for creating comment boxes
 	UEdGraphNode_Comment* CommentTemplate = NewObject<UEdGraphNode_Comment>();
 
 	UBlueprint* Blueprint = FBlueprintEditorUtils::FindBlueprintForGraph(ParentGraph);
 
-	FVector2D SpawnLocation = Location;
+	FVector2f SpawnLocation = Location;
 
 	FSlateRect Bounds;
 	if ((Blueprint != NULL) && FKismetEditorUtilities::GetBoundsForSelectedNodes(Blueprint, Bounds, 50.0f))
@@ -245,7 +245,7 @@ bool UAnimationStateMachineSchema::CreateAutomaticConversionNodeAndConnections(U
 		&& (NodeA->GetInputPin() != NULL) && (NodeA->GetOutputPin() != NULL)
 		&& (NodeB->GetInputPin() != NULL) && (NodeB->GetOutputPin() != NULL))
 	{
-		FVector2D Location = (FVector2D(NodeA->NodePosX, NodeA->NodePosY) + FVector2D(NodeB->NodePosX, NodeB->NodePosY)) * 0.5f;
+		FVector2f Location = (FVector2f(NodeA->NodePosX, NodeA->NodePosY) + FVector2f(NodeB->NodePosX, NodeB->NodePosY)) * 0.5f;
 		UAnimStateTransitionNode* TransitionNode = FEdGraphSchemaAction_NewStateNode::SpawnNodeFromTemplate<UAnimStateTransitionNode>(NodeA->GetGraph(), NewObject<UAnimStateTransitionNode>(), Location, false);
 
 		if (PinA->Direction == EGPD_Output)
@@ -457,7 +457,7 @@ void UAnimationStateMachineSchema::GetGraphDisplayInformation(const UEdGraph& Gr
 	DisplayInfo.Tooltip = LOCTEXT("GraphTooltip_StateMachineSchema", "Graph used to transition between different states each with separate animation graphs");
 }
 
-void UAnimationStateMachineSchema::DroppedAssetsOnGraph(const TArray<FAssetData>& Assets, const FVector2D& GraphPosition, UEdGraph* Graph) const
+void UAnimationStateMachineSchema::DroppedAssetsOnGraph(const TArray<FAssetData>& Assets, const FVector2f& GraphPosition, UEdGraph* Graph) const
 {
 	UAnimationAsset* Asset = FAssetData::GetFirstAsset<UAnimationAsset>(Assets);
 	if(Asset != NULL && GetNodeClassForAsset(Asset->GetClass()))
@@ -469,12 +469,12 @@ void UAnimationStateMachineSchema::DroppedAssetsOnGraph(const TArray<FAssetData>
 		FEdGraphUtilities::RenameGraphToNameOrCloseToName(NewStateNode->BoundGraph, Asset->GetName());
 
 		// Change the current graph context to the inner graph, so that the rest of the drag drop happens inside it
-		FVector2D NewGraphPosition = FVector2D(-300.0f, 0.0f);
+		FVector2f NewGraphPosition = FVector2f(-300.0f, 0.0f);
 		UAnimationGraphSchema::SpawnNodeFromAsset(Asset, NewGraphPosition, NewStateNode->BoundGraph, NewStateNode->GetPoseSinkPinInsideState());
 	}
 }
 
-void UAnimationStateMachineSchema::DroppedAssetsOnNode(const TArray<FAssetData>& Assets, const FVector2D& GraphPosition, UEdGraphNode* Node) const
+void UAnimationStateMachineSchema::DroppedAssetsOnNode(const TArray<FAssetData>& Assets, const FVector2f& GraphPosition, UEdGraphNode* Node) const
 {
 	UAnimationAsset* Asset = FAssetData::GetFirstAsset<UAnimationAsset>(Assets);
 	UAnimStateNode* StateNodeUnderCursor = Cast<UAnimStateNode>(Node);
@@ -499,14 +499,14 @@ void UAnimationStateMachineSchema::DroppedAssetsOnNode(const TArray<FAssetData>&
 			else
 			{
 				// Try dropping onto the pose pin, since nothing is currently connected
-				const FVector2D NewGraphPosition(-300.0f, 0.0f);
+				const FVector2f NewGraphPosition(-300.0f, 0.0f);
 				UAnimationGraphSchema::SpawnNodeFromAsset(Asset, NewGraphPosition, StateNodeUnderCursor->BoundGraph, PosePin);
 			}
 		}
 	}
 }
 
-void UAnimationStateMachineSchema::DroppedAssetsOnPin(const TArray<FAssetData>& Assets, const FVector2D& GraphPosition, UEdGraphPin* Pin) const
+void UAnimationStateMachineSchema::DroppedAssetsOnPin(const TArray<FAssetData>& Assets, const FVector2f& GraphPosition, UEdGraphPin* Pin) const
 {
 	// unused for state machines?
 }
diff --git a/Engine/Source/Editor/AnimGraph/Public/AnimationGraphSchema.h b/Engine/Source/Editor/AnimGraph/Public/AnimationGraphSchema.h
--- a/Engine/Source/Editor/AnimGraph/Public/AnimationGraphSchema.h
+++ b/Engine/Source/Editor/AnimGraph/Public/AnimationGraphSchema.h
@@ -53,9 +53,9 @@ class UAnimationGraphSchema : public UEdGraphSchema_K2
 	ANIMGRAPH_API virtual void CreateDefaultNodesForGraph(UEdGraph& Graph) const override;
 	ANIMGRAPH_API virtual void HandleGraphBeingDeleted(UEdGraph& GraphBeingRemoved) const override;
 	ANIMGRAPH_API virtual bool CreateAutomaticConversionNodeAndConnections(UEdGraphPin* PinA, UEdGraphPin* PinB) const override;
-	ANIMGRAPH_API virtual void DroppedAssetsOnGraph(const TArray<FAssetData>& Assets, const FVector2D& GraphPosition, UEdGraph* Graph) const override;
-	ANIMGRAPH_API virtual void DroppedAssetsOnNode(const TArray<FAssetData>& Assets, const FVector2D& GraphPosition, UEdGraphNode* Node) const override;
-	ANIMGRAPH_API virtual void DroppedAssetsOnPin(const TArray<FAssetData>& Assets, const FVector2D& GraphPosition, UEdGraphPin* Pin) const override;
+	ANIMGRAPH_API virtual void DroppedAssetsOnGraph(const TArray<FAssetData>& Assets, const FVector2f& GraphPosition, UEdGraph* Graph) const override;
+	ANIMGRAPH_API virtual void DroppedAssetsOnNode(const TArray<FAssetData>& Assets, const FVector2f& GraphPosition, UEdGraphNode* Node) const override;
+	ANIMGRAPH_API virtual void DroppedAssetsOnPin(const TArray<FAssetData>& Assets, const FVector2f& GraphPosition, UEdGraphPin* Pin) const override;
 	ANIMGRAPH_API virtual void GetAssetsNodeHoverMessage(const TArray<FAssetData>& Assets, const UEdGraphNode* HoverNode, FString& OutTooltipText, bool& OutOkIcon) const override;
 	ANIMGRAPH_API virtual void GetAssetsPinHoverMessage(const TArray<FAssetData>& Assets, const UEdGraphPin* HoverPin, FString& OutTooltipText, bool& OutOkIcon) const override;
 	ANIMGRAPH_API virtual void GetAssetsGraphHoverMessage(const TArray<FAssetData>& Assets, const UEdGraph* HoverGraph, FString& OutTooltipText, bool& OutOkIcon) const override;
@@ -83,10 +83,10 @@ class UAnimationGraphSchema : public UEdGraphSchema_K2
 	//~ End UEdGraphSchema_K2 Interface
 
 	/** Spawn the correct node in the Animation Graph using the given AnimationAsset at the supplied location */
-	static void SpawnNodeFromAsset(UAnimationAsset* Asset, const FVector2D& GraphPosition, UEdGraph* Graph, UEdGraphPin* PinIfAvailable);
+	static void SpawnNodeFromAsset(UAnimationAsset* Asset, const UE::Slate::FDeprecateVector2DParameter& GraphPosition, UEdGraph* Graph, UEdGraphPin* PinIfAvailable);
 
 	/** Spawn a rigid body node if we drop a physics asset on the graph */
-	static void SpawnRigidBodyNodeFromAsset(UPhysicsAsset* Asset, const FVector2D& GraphPosition, UEdGraph* Graph);
+	static void SpawnRigidBodyNodeFromAsset(UPhysicsAsset* Asset, const UE::Slate::FDeprecateVector2DParameter& GraphPosition, UEdGraph* Graph);
 
 	/** Update the specified node to a new asset */
 	static void UpdateNodeWithAsset(class UK2Node* K2Node, UAnimationAsset* Asset);
diff --git a/Engine/Source/Editor/AnimGraph/Public/AnimationStateMachineSchema.h b/Engine/Source/Editor/AnimGraph/Public/AnimationStateMachineSchema.h
--- a/Engine/Source/Editor/AnimGraph/Public/AnimationStateMachineSchema.h
+++ b/Engine/Source/Editor/AnimGraph/Public/AnimationStateMachineSchema.h
@@ -30,11 +30,11 @@ struct ANIMGRAPH_API FEdGraphSchemaAction_NewStateNode : public FEdGraphSchemaAc
 		, NodeTemplate(nullptr)
 	{}
 
-	virtual UEdGraphNode* PerformAction(class UEdGraph* ParentGraph, UEdGraphPin* FromPin, const FVector2D Location, bool bSelectNewNode = true) override;
+	virtual UEdGraphNode* PerformAction(class UEdGraph* ParentGraph, UEdGraphPin* FromPin, const FVector2f& Location, bool bSelectNewNode = true) override;
 	virtual void AddReferencedObjects( FReferenceCollector& Collector ) override;
 
 	template <typename NodeType>
-	static NodeType* SpawnNodeFromTemplate(class UEdGraph* ParentGraph, NodeType* InTemplateNode, const FVector2D Location = FVector2D(0.0f, 0.0f), bool bSelectNewNode = true)
+	static NodeType* SpawnNodeFromTemplate(class UEdGraph* ParentGraph, NodeType* InTemplateNode, const UE::Slate::FDeprecateVector2DParameter Location = FVector2f(0.0f, 0.0f), bool bSelectNewNode = true)
 	{
 		FEdGraphSchemaAction_NewStateNode Action;
 		Action.NodeTemplate = InTemplateNode;
@@ -58,7 +58,7 @@ struct ANIMGRAPH_API FEdGraphSchemaAction_NewStateComment : public FEdGraphSchem
 	{}
 
 	//~ Begin FEdGraphSchemaAction Interface
-	virtual UEdGraphNode* PerformAction(class UEdGraph* ParentGraph, UEdGraphPin* FromPin, const FVector2D Location, bool bSelectNewNode = true) override;
+	virtual UEdGraphNode* PerformAction(class UEdGraph* ParentGraph, UEdGraphPin* FromPin, const FVector2f& Location, bool bSelectNewNode = true) override;
 	//~ End FEdGraphSchemaAction Interface
 };
 
@@ -83,9 +83,9 @@ class UAnimationStateMachineSchema : public UEdGraphSchema
 	virtual void GetContextMenuActions(class UToolMenu* Menu, class UGraphNodeContextMenuContext* Context) const override;
 	virtual FLinearColor GetPinTypeColor(const FEdGraphPinType& PinType) const override;
 	virtual void GetGraphDisplayInformation(const UEdGraph& Graph, /*out*/ FGraphDisplayInfo& DisplayInfo) const override;
-	virtual void DroppedAssetsOnGraph(const TArray<FAssetData>& Assets, const FVector2D& GraphPosition, UEdGraph* Graph) const override;
-	virtual void DroppedAssetsOnNode(const TArray<FAssetData>& Assets, const FVector2D& GraphPosition, UEdGraphNode* Node) const override;
-	virtual void DroppedAssetsOnPin(const TArray<FAssetData>& Assets, const FVector2D& GraphPosition, UEdGraphPin* Pin) const override;
+	virtual void DroppedAssetsOnGraph(const TArray<FAssetData>& Assets, const FVector2f& GraphPosition, UEdGraph* Graph) const override;
+	virtual void DroppedAssetsOnNode(const TArray<FAssetData>& Assets, const FVector2f& GraphPosition, UEdGraphNode* Node) const override;
+	virtual void DroppedAssetsOnPin(const TArray<FAssetData>& Assets, const FVector2f& GraphPosition, UEdGraphPin* Pin) const override;
 	virtual void GetAssetsNodeHoverMessage(const TArray<FAssetData>& Assets, const UEdGraphNode* HoverNode, FString& OutTooltipText, bool& OutOkIcon) const override;
 	virtual void GetAssetsPinHoverMessage(const TArray<FAssetData>& Assets, const UEdGraphPin* HoverPin, FString& OutTooltipText, bool& OutOkIcon) const override;
 
diff --git a/Engine/Source/Editor/AnimationBlueprintEditor/Private/AnimationBlueprintEditor.cpp b/Engine/Source/Editor/AnimationBlueprintEditor/Private/AnimationBlueprintEditor.cpp
--- a/Engine/Source/Editor/AnimationBlueprintEditor/Private/AnimationBlueprintEditor.cpp
+++ b/Engine/Source/Editor/AnimationBlueprintEditor/Private/AnimationBlueprintEditor.cpp
@@ -2010,7 +2010,7 @@ void FAnimationBlueprintEditor::OnCreateComment()
 		if (UEdGraph* Graph = GraphEditor->GetCurrentGraph())
 		{
 			FEdGraphSchemaAction_K2AddComment CommentAction;
-			CommentAction.PerformAction(Graph, nullptr, GraphEditor->GetPasteLocation());
+			CommentAction.PerformAction(Graph, nullptr, GraphEditor->GetPasteLocation2f());
 		}
 	}
 }
diff --git a/Engine/Source/Editor/AudioEditor/Classes/SoundClassGraph/SoundClassGraphSchema.h b/Engine/Source/Editor/AudioEditor/Classes/SoundClassGraph/SoundClassGraphSchema.h
--- a/Engine/Source/Editor/AudioEditor/Classes/SoundClassGraph/SoundClassGraphSchema.h
+++ b/Engine/Source/Editor/AudioEditor/Classes/SoundClassGraph/SoundClassGraphSchema.h
@@ -45,7 +45,7 @@ struct AUDIOEDITOR_API FSoundClassGraphSchemaAction_NewNode : public FEdGraphSch
 
 	//~ Begin FEdGraphSchemaAction Interface
 	virtual FName GetTypeId() const override { return StaticGetTypeId(); } 
-	virtual UEdGraphNode* PerformAction(class UEdGraph* ParentGraph, UEdGraphPin* FromPin, const FVector2D Location, bool bSelectNewNode = true) override;
+	virtual UEdGraphNode* PerformAction(class UEdGraph* ParentGraph, UEdGraphPin* FromPin, const FVector2f& Location, bool bSelectNewNode = true) override;
 	//~ End FEdGraphSchemaAction Interface
 
 	/** Name for the new SoundClass */
@@ -71,7 +71,7 @@ class USoundClassGraphSchema : public UEdGraphSchema
 	virtual void BreakNodeLinks(UEdGraphNode& TargetNode) const override;
 	virtual void BreakPinLinks(UEdGraphPin& TargetPin, bool bSendsNodeNotifcation) const override;
 	virtual void BreakSinglePinLink(UEdGraphPin* SourcePin, UEdGraphPin* TargetPin) const override;
-	virtual void DroppedAssetsOnGraph(const TArray<struct FAssetData>& Assets, const FVector2D& GraphPosition, UEdGraph* Graph) const override;
+	virtual void DroppedAssetsOnGraph(const TArray<struct FAssetData>& Assets, const FVector2f& GraphPosition, UEdGraph* Graph) const override;
 	//~ End EdGraphSchema Interface
 };
 
diff --git a/Engine/Source/Editor/AudioEditor/Classes/SoundCueGraph/SoundCueGraphSchema.h b/Engine/Source/Editor/AudioEditor/Classes/SoundCueGraph/SoundCueGraphSchema.h
--- a/Engine/Source/Editor/AudioEditor/Classes/SoundCueGraph/SoundCueGraphSchema.h
+++ b/Engine/Source/Editor/AudioEditor/Classes/SoundCueGraph/SoundCueGraphSchema.h
@@ -47,7 +47,8 @@ struct AUDIOEDITOR_API FSoundCueGraphSchemaAction_NewNode : public FEdGraphSchem
 	{}
 
 	//~ Begin FEdGraphSchemaAction Interface
-	virtual UEdGraphNode* PerformAction(class UEdGraph* ParentGraph, UEdGraphPin* FromPin, const FVector2D Location, bool bSelectNewNode = true) override;
+	using FEdGraphSchemaAction::PerformAction; // Prevent hiding of deprecated base class function with FVector2D
+	virtual UEdGraphNode* PerformAction(class UEdGraph* ParentGraph, UEdGraphPin* FromPin, const FVector2f& Location, bool bSelectNewNode = true) override;
 	//~ End FEdGraphSchemaAction Interface
 
 private:
@@ -70,7 +71,8 @@ struct AUDIOEDITOR_API FSoundCueGraphSchemaAction_NewFromSelected : public FSoun
 	{}
 
 	//~ Begin FEdGraphSchemaAction Interface
-	virtual UEdGraphNode* PerformAction(class UEdGraph* ParentGraph, UEdGraphPin* FromPin, const FVector2D Location, bool bSelectNewNode = true) override;
+	using FEdGraphSchemaAction::PerformAction; // Prevent hiding of deprecated base class function with FVector2D
+	virtual UEdGraphNode* PerformAction(class UEdGraph* ParentGraph, UEdGraphPin* FromPin, const FVector2f& Location, bool bSelectNewNode = true) override;
 	//~ End FEdGraphSchemaAction Interface
 };
 
@@ -89,7 +91,8 @@ struct AUDIOEDITOR_API FSoundCueGraphSchemaAction_NewComment : public FEdGraphSc
 	{}
 
 	//~ Begin FEdGraphSchemaAction Interface
-	virtual UEdGraphNode* PerformAction(class UEdGraph* ParentGraph, UEdGraphPin* FromPin, const FVector2D Location, bool bSelectNewNode = true) override;
+	using FEdGraphSchemaAction::PerformAction; // Prevent hiding of deprecated base class function with FVector2D
+	virtual UEdGraphNode* PerformAction(class UEdGraph* ParentGraph, UEdGraphPin* FromPin, const FVector2f& Location, bool bSelectNewNode = true) override;
 	//~ End FEdGraphSchemaAction Interface
 };
 
@@ -108,7 +111,8 @@ struct AUDIOEDITOR_API FSoundCueGraphSchemaAction_Paste : public FEdGraphSchemaA
 	{}
 
 	//~ Begin FEdGraphSchemaAction Interface
-	virtual UEdGraphNode* PerformAction(class UEdGraph* ParentGraph, UEdGraphPin* FromPin, const FVector2D Location, bool bSelectNewNode = true) override;
+	using FEdGraphSchemaAction::PerformAction; // Prevent hiding of deprecated base class function with FVector2D
+	virtual UEdGraphNode* PerformAction(class UEdGraph* ParentGraph, UEdGraphPin* FromPin, const FVector2f& Location, bool bSelectNewNode = true) override;
 	//~ End FEdGraphSchemaAction Interface
 };
 
@@ -137,8 +141,8 @@ class USoundCueGraphSchema : public UEdGraphSchema
 	virtual void BreakNodeLinks(UEdGraphNode& TargetNode) const override;
 	virtual void BreakPinLinks(UEdGraphPin& TargetPin, bool bSendsNodeNotifcation) const override;
 	virtual void GetAssetsGraphHoverMessage(const TArray<FAssetData>& Assets, const UEdGraph* HoverGraph, FString& OutTooltipText, bool& OutOkIcon) const override;
-	virtual void DroppedAssetsOnGraph(const TArray<struct FAssetData>& Assets, const FVector2D& GraphPosition, UEdGraph* Graph) const override;
-	virtual void DroppedAssetsOnNode(const TArray<FAssetData>& Assets, const FVector2D& GraphPosition, UEdGraphNode* Node) const override;
+	virtual void DroppedAssetsOnGraph(const TArray<struct FAssetData>& Assets, const FVector2f& GraphPosition, UEdGraph* Graph) const override;
+	virtual void DroppedAssetsOnNode(const TArray<FAssetData>& Assets, const FVector2f& GraphPosition, UEdGraphNode* Node) const override;
 	virtual int32 GetNodeSelectionCount(const UEdGraph* Graph) const override;
 	virtual TSharedPtr<FEdGraphSchemaAction> GetCreateCommentAction() const override;
 	//~ End EdGraphSchema Interface
diff --git a/Engine/Source/Editor/AudioEditor/Classes/SoundSubmixGraph/SoundSubmixGraphSchema.h b/Engine/Source/Editor/AudioEditor/Classes/SoundSubmixGraph/SoundSubmixGraphSchema.h
--- a/Engine/Source/Editor/AudioEditor/Classes/SoundSubmixGraph/SoundSubmixGraphSchema.h
+++ b/Engine/Source/Editor/AudioEditor/Classes/SoundSubmixGraph/SoundSubmixGraphSchema.h
@@ -110,7 +110,7 @@ struct AUDIOEDITOR_API FSoundSubmixGraphSchemaAction_NewNode : public FEdGraphSc
 
 	//~ Begin FEdGraphSchemaAction Interface
 	virtual FName GetTypeId() const override { return StaticGetTypeId(); } 
-	virtual UEdGraphNode* PerformAction(class UEdGraph* ParentGraph, UEdGraphPin* FromPin, const FVector2D Location, bool bSelectNewNode = true) override;
+	virtual UEdGraphNode* PerformAction(class UEdGraph* ParentGraph, UEdGraphPin* FromPin, const FVector2f& Location, bool bSelectNewNode = true) override;
 	//~ End FEdGraphSchemaAction Interface
 
 	/** Name for the new SoundSubmix */
@@ -136,7 +136,7 @@ class USoundSubmixGraphSchema : public UEdGraphSchema
 	virtual void BreakNodeLinks(UEdGraphNode& TargetNode) const override;
 	virtual void BreakPinLinks(UEdGraphPin& TargetPin, bool bSendsNodeNotifcation) const override;
 	virtual void BreakSinglePinLink(UEdGraphPin* SourcePin, UEdGraphPin* TargetPin) const override;
-	virtual void DroppedAssetsOnGraph(const TArray<struct FAssetData>& Assets, const FVector2D& GraphPosition, UEdGraph* Graph) const override;
+	virtual void DroppedAssetsOnGraph(const TArray<struct FAssetData>& Assets, const FVector2f& GraphPosition, UEdGraph* Graph) const override;
 	virtual void GetAssetsGraphHoverMessage(const TArray<FAssetData>& Assets, const UEdGraph* HoverGraph, FString& OutTooltipText, bool& OutOkIcon) const override;
 	//~ End EdGraphSchema Interface
 };
diff --git a/Engine/Source/Editor/AudioEditor/Private/SoundClassGraphSchema.cpp b/Engine/Source/Editor/AudioEditor/Private/SoundClassGraphSchema.cpp
--- a/Engine/Source/Editor/AudioEditor/Private/SoundClassGraphSchema.cpp
+++ b/Engine/Source/Editor/AudioEditor/Private/SoundClassGraphSchema.cpp
@@ -23,9 +23,9 @@
 
 #define LOCTEXT_NAMESPACE "SoundClassSchema"
 
-UEdGraphNode* FSoundClassGraphSchemaAction_NewNode::PerformAction(class UEdGraph* ParentGraph, UEdGraphPin* FromPin, const FVector2D Location, bool bSelectNewNode/* = true*/)
+UEdGraphNode* FSoundClassGraphSchemaAction_NewNode::PerformAction(class UEdGraph* ParentGraph, UEdGraphPin* FromPin, const FVector2f& Location, bool bSelectNewNode/* = true*/)
 {
-	FSoundClassEditorUtilities::CreateSoundClass(ParentGraph, FromPin, Location, NewSoundClassName);
+	FSoundClassEditorUtilities::CreateSoundClass(ParentGraph, FromPin, FDeprecateSlateVector2D(Location), NewSoundClassName);
 	return NULL;
 }
 
@@ -158,7 +158,7 @@ void USoundClassGraphSchema::BreakSinglePinLink(UEdGraphPin* SourcePin, UEdGraph
 	CastChecked<USoundClassGraph>(SourcePin->GetOwningNode()->GetGraph())->LinkSoundClasses();
 }
 
-void USoundClassGraphSchema::DroppedAssetsOnGraph(const TArray<struct FAssetData>& Assets, const FVector2D& GraphPosition, UEdGraph* Graph) const
+void USoundClassGraphSchema::DroppedAssetsOnGraph(const TArray<struct FAssetData>& Assets, const FVector2f& GraphPosition, UEdGraph* Graph) const
 {
 	USoundClassGraph* SoundClassGraph = CastChecked<USoundClassGraph>(Graph);
 
diff --git a/Engine/Source/Editor/AudioEditor/Private/SoundCueEditor.cpp b/Engine/Source/Editor/AudioEditor/Private/SoundCueEditor.cpp
--- a/Engine/Source/Editor/AudioEditor/Private/SoundCueEditor.cpp
+++ b/Engine/Source/Editor/AudioEditor/Private/SoundCueEditor.cpp
@@ -609,7 +609,7 @@ bool FSoundCueEditor::CanDeleteInput() const
 void FSoundCueEditor::OnCreateComment()
 {
 	FSoundCueGraphSchemaAction_NewComment CommentAction;
-	CommentAction.PerformAction(SoundCue->SoundCueGraph, NULL, SoundCueGraphEditor->GetPasteLocation());
+	CommentAction.PerformAction(SoundCue->SoundCueGraph, NULL, SoundCueGraphEditor->GetPasteLocation2f());
 }
 
 TSharedRef<SGraphEditor> FSoundCueEditor::CreateGraphEditorWidget()
@@ -938,7 +938,7 @@ bool FSoundCueEditor::CanCopyNodes() const
 
 void FSoundCueEditor::PasteNodes()
 {
-	PasteNodesHere(SoundCueGraphEditor->GetPasteLocation());
+	PasteNodesHere(SoundCueGraphEditor->GetPasteLocation2f());
 }
 
 void FSoundCueEditor::PasteNodesHere(const FVector2D& Location)
diff --git a/Engine/Source/Editor/AudioEditor/Private/SoundCueGraphSchema.cpp b/Engine/Source/Editor/AudioEditor/Private/SoundCueGraphSchema.cpp
--- a/Engine/Source/Editor/AudioEditor/Private/SoundCueGraphSchema.cpp
+++ b/Engine/Source/Editor/AudioEditor/Private/SoundCueGraphSchema.cpp
@@ -48,7 +48,7 @@ class FString;
 /////////////////////////////////////////////////////
 // FSoundCueGraphSchemaAction_NewNode
 
-UEdGraphNode* FSoundCueGraphSchemaAction_NewNode::PerformAction(class UEdGraph* ParentGraph, UEdGraphPin* FromPin, const FVector2D Location, bool bSelectNewNode/* = true*/)
+UEdGraphNode* FSoundCueGraphSchemaAction_NewNode::PerformAction(class UEdGraph* ParentGraph, UEdGraphPin* FromPin, const FVector2f& Location, bool bSelectNewNode/* = true*/)
 {
 	check(SoundNodeClass);
 
@@ -123,7 +123,7 @@ void FSoundCueGraphSchemaAction_NewNode::ConnectToSelectedNodes(USoundNode* NewN
 /////////////////////////////////////////////////////
 // FSoundCueGraphSchemaAction_NewFromSelected
 
-UEdGraphNode* FSoundCueGraphSchemaAction_NewFromSelected::PerformAction(class UEdGraph* ParentGraph, UEdGraphPin* FromPin, const FVector2D Location, bool bSelectNewNode/* = true*/)
+UEdGraphNode* FSoundCueGraphSchemaAction_NewFromSelected::PerformAction(class UEdGraph* ParentGraph, UEdGraphPin* FromPin, const FVector2f& Location, bool bSelectNewNode/* = true*/)
 {
 	USoundCue* SoundCue = CastChecked<USoundCueGraph>(ParentGraph)->GetSoundCue();
 	const FScopedTransaction Transaction( LOCTEXT("SoundCueEditorNewFromSelection", "Sound Cue Editor: New From Selection") );
@@ -132,12 +132,12 @@ UEdGraphNode* FSoundCueGraphSchemaAction_NewFromSelected::PerformAction(class UE
 
 	UEdGraphNode* CreatedNode = NULL;
 
-	FVector2D WaveStartLocation = Location;
+	FDeprecateSlateVector2D WaveStartLocation = Location;
 
 	if (SoundNodeClass)
 	{
 		// If we will create another node, move wave nodes out of the way.
-		WaveStartLocation.X -= 200;
+		WaveStartLocation.X -= 200.0f;
 	}
 
 	TArray<USoundWave*> SelectedWaves;
@@ -191,12 +191,12 @@ UEdGraphNode* FSoundCueGraphSchemaAction_NewFromSelected::PerformAction(class UE
 /////////////////////////////////////////////////////
 // FSoundCueGraphSchemaAction_NewComment
 
-UEdGraphNode* FSoundCueGraphSchemaAction_NewComment::PerformAction(class UEdGraph* ParentGraph, UEdGraphPin* FromPin, const FVector2D Location, bool bSelectNewNode/* = true*/)
+UEdGraphNode* FSoundCueGraphSchemaAction_NewComment::PerformAction(class UEdGraph* ParentGraph, UEdGraphPin* FromPin, const FVector2f& Location, bool bSelectNewNode/* = true*/)
 {
 	// Add menu item for creating comment boxes
 	UEdGraphNode_Comment* CommentTemplate = NewObject<UEdGraphNode_Comment>();
 
-	FVector2D SpawnLocation = Location;
+	FVector2f SpawnLocation = Location;
 
 	FSlateRect Bounds;
 	if (FSoundCueEditorUtilities::GetBoundsForSelectedNodes(ParentGraph, Bounds, 50.0f))
@@ -212,9 +212,9 @@ UEdGraphNode* FSoundCueGraphSchemaAction_NewComment::PerformAction(class UEdGrap
 /////////////////////////////////////////////////////
 // FSoundCueGraphSchemaAction_Paste
 
-UEdGraphNode* FSoundCueGraphSchemaAction_Paste::PerformAction(class UEdGraph* ParentGraph, UEdGraphPin* FromPin, const FVector2D Location, bool bSelectNewNode/* = true*/)
+UEdGraphNode* FSoundCueGraphSchemaAction_Paste::PerformAction(class UEdGraph* ParentGraph, UEdGraphPin* FromPin, const FVector2f& Location, bool bSelectNewNode/* = true*/)
 {
-	FSoundCueEditorUtilities::PasteNodesHere(ParentGraph, Location);
+	FSoundCueEditorUtilities::PasteNodesHere(ParentGraph, FDeprecateSlateVector2D(Location));
 	return NULL;
 }
 
@@ -431,7 +431,7 @@ void USoundCueGraphSchema::GetAssetsGraphHoverMessage(const TArray<FAssetData>&
 	}
 }
 
-void USoundCueGraphSchema::DroppedAssetsOnGraph(const TArray<FAssetData>& Assets, const FVector2D& GraphPosition, UEdGraph* Graph) const
+void USoundCueGraphSchema::DroppedAssetsOnGraph(const TArray<FAssetData>& Assets, const FVector2f& GraphPosition, UEdGraph* Graph) const
 {
 	//////////////////////////////////////////////////////////////////////////
 	// Handle dropped USoundWaves
@@ -455,7 +455,7 @@ void USoundCueGraphSchema::DroppedAssetsOnGraph(const TArray<FAssetData>& Assets
 		SoundCueGraph->Modify();
 
 		TArray<USoundNode*> CreatedPlayers;
-		FSoundCueEditorUtilities::CreateWaveContainers(Waves, SoundCue, CreatedPlayers, GraphPosition);
+		FSoundCueEditorUtilities::CreateWaveContainers(Waves, SoundCue, CreatedPlayers, FDeprecateSlateVector2D(GraphPosition));
 	}
 
 	//////////////////////////////////////////////////////////////////////////
@@ -480,11 +480,11 @@ void USoundCueGraphSchema::DroppedAssetsOnGraph(const TArray<FAssetData>& Assets
 		SoundCueGraph->Modify();
 
 		TArray<USoundNode*> CreatedPlayers;
-		FSoundCueEditorUtilities::CreateDialogueContainers(Dialogues, SoundCue, CreatedPlayers, GraphPosition);
+		FSoundCueEditorUtilities::CreateDialogueContainers(Dialogues, SoundCue, CreatedPlayers, FDeprecateSlateVector2D(GraphPosition));
 	}
 }
 
-void USoundCueGraphSchema::DroppedAssetsOnNode(const TArray<FAssetData>& Assets, const FVector2D& GraphPosition, UEdGraphNode* Node) const
+void USoundCueGraphSchema::DroppedAssetsOnNode(const TArray<FAssetData>& Assets, const FVector2f& GraphPosition, UEdGraphNode* Node) const
 {
 	// Currently, drag and drop is only supported for dropping on sound cue graph nodes, and in particular, sound wave players and sound dialogue players.
 	if (!Node->IsA<USoundCueGraphNode>())
@@ -529,13 +529,13 @@ void USoundCueGraphSchema::DroppedAssetsOnNode(const TArray<FAssetData>& Assets,
 			for (int32 Index = 1; Index < Waves.Num(); Index++)
 			{
 				TArray<USoundNode*> CreatedPlayers;
-				FSoundCueEditorUtilities::CreateWaveContainers(Waves, SoundCue, CreatedPlayers, GraphPosition);
+				FSoundCueEditorUtilities::CreateWaveContainers(Waves, SoundCue, CreatedPlayers, FDeprecateSlateVector2D(GraphPosition));
 			}
 		}
 		else if (Dialogues.Num() > 0)
 		{
 			TArray<USoundNode*> CreatedPlayers;
-			FSoundCueEditorUtilities::CreateDialogueContainers(Dialogues, SoundCue, CreatedPlayers, GraphPosition);
+			FSoundCueEditorUtilities::CreateDialogueContainers(Dialogues, SoundCue, CreatedPlayers, FDeprecateSlateVector2D(GraphPosition));
 
 			if (CreatedPlayers.Num() > 0)
 			{
@@ -572,13 +572,13 @@ void USoundCueGraphSchema::DroppedAssetsOnNode(const TArray<FAssetData>& Assets,
 			for (int32 Index = 1; Index < Waves.Num(); Index++)
 			{
 				TArray<USoundNode*> CreatedPlayers;
-				FSoundCueEditorUtilities::CreateDialogueContainers(Dialogues, SoundCue, CreatedPlayers, GraphPosition);
+				FSoundCueEditorUtilities::CreateDialogueContainers(Dialogues, SoundCue, CreatedPlayers, FDeprecateSlateVector2D(GraphPosition));
 			}
 		}
 		else if (Waves.Num() > 0)
 		{
 			TArray<USoundNode*> CreatedPlayers;
-			FSoundCueEditorUtilities::CreateWaveContainers(Waves, SoundCue, CreatedPlayers, GraphPosition);
+			FSoundCueEditorUtilities::CreateWaveContainers(Waves, SoundCue, CreatedPlayers, FDeprecateSlateVector2D(GraphPosition));
 
 			if (CreatedPlayers.Num() > 0)
 			{
diff --git a/Engine/Source/Editor/AudioEditor/Private/SoundSubmixGraphSchema.cpp b/Engine/Source/Editor/AudioEditor/Private/SoundSubmixGraphSchema.cpp
--- a/Engine/Source/Editor/AudioEditor/Private/SoundSubmixGraphSchema.cpp
+++ b/Engine/Source/Editor/AudioEditor/Private/SoundSubmixGraphSchema.cpp
@@ -106,9 +106,9 @@ void FSoundSubmixGraphConnectionDrawingPolicy::DetermineWiringStyle(UEdGraphPin*
 	}
 }
 
-UEdGraphNode* FSoundSubmixGraphSchemaAction_NewNode::PerformAction(UEdGraph* ParentGraph, UEdGraphPin* FromPin, const FVector2D Location, bool bSelectNewNode/* = true*/)
+UEdGraphNode* FSoundSubmixGraphSchemaAction_NewNode::PerformAction(UEdGraph* ParentGraph, UEdGraphPin* FromPin, const FVector2f& Location, bool bSelectNewNode/* = true*/)
 {
-	FSoundSubmixEditorUtilities::CreateSoundSubmix(ParentGraph, FromPin, Location, NewSoundSubmixName);
+	FSoundSubmixEditorUtilities::CreateSoundSubmix(ParentGraph, FromPin, FDeprecateSlateVector2D(Location), NewSoundSubmixName);
 	return nullptr;
 }
 
@@ -456,7 +456,7 @@ void USoundSubmixGraphSchema::BreakSinglePinLink(UEdGraphPin* SourcePin, UEdGrap
 	CastChecked<USoundSubmixGraph>(SourcePin->GetOwningNode()->GetGraph())->LinkSoundSubmixes();
 }
 
-void USoundSubmixGraphSchema::DroppedAssetsOnGraph(const TArray<FAssetData>& Assets, const FVector2D& GraphPosition, UEdGraph* Graph) const
+void USoundSubmixGraphSchema::DroppedAssetsOnGraph(const TArray<FAssetData>& Assets, const FVector2f& GraphPosition, UEdGraph* Graph) const
 {
 	check(GEditor);
 	check(Graph);
diff --git a/Engine/Source/Editor/BehaviorTreeEditor/Classes/EdGraphSchema_BehaviorTree.h b/Engine/Source/Editor/BehaviorTreeEditor/Classes/EdGraphSchema_BehaviorTree.h
--- a/Engine/Source/Editor/BehaviorTreeEditor/Classes/EdGraphSchema_BehaviorTree.h
+++ b/Engine/Source/Editor/BehaviorTreeEditor/Classes/EdGraphSchema_BehaviorTree.h
@@ -45,7 +45,7 @@ struct FBehaviorTreeSchemaAction_AutoArrange : public FEdGraphSchemaAction
 	{}
 
 	//~ Begin FEdGraphSchemaAction Interface
-	virtual UEdGraphNode* PerformAction(class UEdGraph* ParentGraph, UEdGraphPin* FromPin, const FVector2D Location, bool bSelectNewNode = true) override;
+	virtual UEdGraphNode* PerformAction(class UEdGraph* ParentGraph, UEdGraphPin* FromPin, const FVector2f& Location, bool bSelectNewNode = true) override;
 	//~ End FEdGraphSchemaAction Interface
 };
 
diff --git a/Engine/Source/Editor/BehaviorTreeEditor/Classes/EdGraphSchema_BehaviorTreeDecorator.h b/Engine/Source/Editor/BehaviorTreeEditor/Classes/EdGraphSchema_BehaviorTreeDecorator.h
--- a/Engine/Source/Editor/BehaviorTreeEditor/Classes/EdGraphSchema_BehaviorTreeDecorator.h
+++ b/Engine/Source/Editor/BehaviorTreeEditor/Classes/EdGraphSchema_BehaviorTreeDecorator.h
@@ -48,13 +48,13 @@ struct FDecoratorSchemaAction_NewNode : public FEdGraphSchemaAction
 	{}
 
 	//~ Begin FEdGraphSchemaAction Interface
-	virtual UEdGraphNode* PerformAction(class UEdGraph* ParentGraph, UEdGraphPin* FromPin, const FVector2D Location, bool bSelectNewNode = true) override;
-	virtual UEdGraphNode* PerformAction(class UEdGraph* ParentGraph, TArray<UEdGraphPin*>& FromPins, const FVector2D Location, bool bSelectNewNode = true) override;
+	virtual UEdGraphNode* PerformAction(class UEdGraph* ParentGraph, UEdGraphPin* FromPin, const FVector2f& Location, bool bSelectNewNode = true) override;
+	virtual UEdGraphNode* PerformAction(class UEdGraph* ParentGraph, TArray<UEdGraphPin*>& FromPins, const FVector2f& Location, bool bSelectNewNode = true) override;
 	virtual void AddReferencedObjects( FReferenceCollector& Collector ) override;
 	//~ End FEdGraphSchemaAction Interface
 
 	template <typename NodeType>
-	static NodeType* SpawnNodeFromTemplate(class UEdGraph* ParentGraph, NodeType* InTemplateNode, const FVector2D Location)
+	static NodeType* SpawnNodeFromTemplate(class UEdGraph* ParentGraph, NodeType* InTemplateNode, const FVector2f& Location)
 	{
 		FDecoratorSchemaAction_NewNode Action;
 		Action.NodeTemplate = InTemplateNode;
diff --git a/Engine/Source/Editor/BehaviorTreeEditor/Private/BehaviorTreeEditorTabFactories.cpp b/Engine/Source/Editor/BehaviorTreeEditor/Private/BehaviorTreeEditorTabFactories.cpp
--- a/Engine/Source/Editor/BehaviorTreeEditor/Private/BehaviorTreeEditorTabFactories.cpp
+++ b/Engine/Source/Editor/BehaviorTreeEditor/Private/BehaviorTreeEditorTabFactories.cpp
@@ -181,7 +181,7 @@ void FBTGraphEditorSummoner::SaveState(TSharedPtr<SDockTab> Tab, TSharedPtr<FTab
 
 	TSharedRef<SGraphEditor> GraphEditor = StaticCastSharedRef<SGraphEditor>(Tab->GetContent());
 
-	FVector2D ViewLocation;
+	FVector2f ViewLocation;
 	float ZoomAmount;
 	GraphEditor->GetViewLocation(ViewLocation, ZoomAmount);
 
diff --git a/Engine/Source/Editor/BehaviorTreeEditor/Private/EdGraphSchema_BehaviorTree.cpp b/Engine/Source/Editor/BehaviorTreeEditor/Private/EdGraphSchema_BehaviorTree.cpp
--- a/Engine/Source/Editor/BehaviorTreeEditor/Private/EdGraphSchema_BehaviorTree.cpp
+++ b/Engine/Source/Editor/BehaviorTreeEditor/Private/EdGraphSchema_BehaviorTree.cpp
@@ -49,7 +49,7 @@ int32 UEdGraphSchema_BehaviorTree::CurrentCacheRefreshID = 0;
 //----------------------------------------------------------------------//
 // 
 //----------------------------------------------------------------------//
-UEdGraphNode* FBehaviorTreeSchemaAction_AutoArrange::PerformAction(class UEdGraph* ParentGraph, UEdGraphPin* FromPin, const FVector2D Location, bool bSelectNewNode)
+UEdGraphNode* FBehaviorTreeSchemaAction_AutoArrange::PerformAction(class UEdGraph* ParentGraph, UEdGraphPin* FromPin, const FVector2f& Location, bool bSelectNewNode)
 {
 	UBehaviorTreeGraph* Graph = Cast<UBehaviorTreeGraph>(ParentGraph);
 	if (Graph)
diff --git a/Engine/Source/Editor/BehaviorTreeEditor/Private/EdGraphSchema_BehaviorTreeDecorator.cpp b/Engine/Source/Editor/BehaviorTreeEditor/Private/EdGraphSchema_BehaviorTreeDecorator.cpp
--- a/Engine/Source/Editor/BehaviorTreeEditor/Private/EdGraphSchema_BehaviorTreeDecorator.cpp
+++ b/Engine/Source/Editor/BehaviorTreeEditor/Private/EdGraphSchema_BehaviorTreeDecorator.cpp
@@ -26,7 +26,7 @@ namespace DecoratorSchema
 	const int32 NodeDistance = 60;
 }
 
-UEdGraphNode* FDecoratorSchemaAction_NewNode::PerformAction(class UEdGraph* ParentGraph, UEdGraphPin* FromPin, const FVector2D Location, bool bSelectNewNode)
+UEdGraphNode* FDecoratorSchemaAction_NewNode::PerformAction(class UEdGraph* ParentGraph, UEdGraphPin* FromPin, const FVector2f& Location, bool bSelectNewNode)
 {
 	UEdGraphNode* ResultNode = NULL;
 
@@ -78,7 +78,7 @@ UEdGraphNode* FDecoratorSchemaAction_NewNode::PerformAction(class UEdGraph* Pare
 	return ResultNode;
 }
 
-UEdGraphNode* FDecoratorSchemaAction_NewNode::PerformAction(class UEdGraph* ParentGraph, TArray<UEdGraphPin*>& FromPins, const FVector2D Location, bool bSelectNewNode) 
+UEdGraphNode* FDecoratorSchemaAction_NewNode::PerformAction(class UEdGraph* ParentGraph, TArray<UEdGraphPin*>& FromPins, const FVector2f& Location, bool bSelectNewNode) 
 {
 	UEdGraphNode* ResultNode = NULL;
 	if (FromPins.Num() > 0)
diff --git a/Engine/Source/Editor/BlueprintGraph/Classes/EdGraphSchema_K2.h b/Engine/Source/Editor/BlueprintGraph/Classes/EdGraphSchema_K2.h
--- a/Engine/Source/Editor/BlueprintGraph/Classes/EdGraphSchema_K2.h
+++ b/Engine/Source/Editor/BlueprintGraph/Classes/EdGraphSchema_K2.h
@@ -557,9 +557,9 @@ public:
 	virtual bool CanEncapuslateNode(UEdGraphNode const& TestNode) const override;
 	virtual void HandleGraphBeingDeleted(UEdGraph& GraphBeingRemoved) const override;
 	virtual void GetGraphDisplayInformation(const UEdGraph& Graph, /*out*/ FGraphDisplayInfo& DisplayInfo) const override;
-	virtual void DroppedAssetsOnGraph(const TArray<FAssetData>& Assets, const FVector2D& GraphPosition, UEdGraph* Graph) const override;
-	virtual void DroppedAssetsOnNode(const TArray<FAssetData>& Assets, const FVector2D& GraphPosition, UEdGraphNode* Node) const override;
-	virtual void DroppedAssetsOnPin(const TArray<FAssetData>& Assets, const FVector2D& GraphPosition, UEdGraphPin* Pin) const override;
+	virtual void DroppedAssetsOnGraph(const TArray<FAssetData>& Assets, const FVector2f& GraphPosition, UEdGraph* Graph) const override;
+	virtual void DroppedAssetsOnNode(const TArray<FAssetData>& Assets, const FVector2f& GraphPosition, UEdGraphNode* Node) const override;
+	virtual void DroppedAssetsOnPin(const TArray<FAssetData>& Assets, const FVector2f& GraphPosition, UEdGraphPin* Pin) const override;
 	virtual void GetAssetsNodeHoverMessage(const TArray<FAssetData>& Assets, const UEdGraphNode* HoverNode, FString& OutTooltipText, bool& OutOkIcon) const override;
 	virtual void GetAssetsPinHoverMessage(const TArray<FAssetData>& Assets, const UEdGraphPin* HoverPin, FString& OutTooltipText, bool& OutOkIcon) const override;
 	virtual void GetAssetsGraphHoverMessage(const TArray<FAssetData>& Assets, const UEdGraph* HoverGraph, FString& OutTooltipText, bool& OutOkIcon) const override;
@@ -573,7 +573,7 @@ public:
 	virtual bool ShouldAlwaysPurgeOnModification() const override { return false; }
 	virtual void SplitPin(UEdGraphPin* Pin, bool bNotify = true) const override;
 	virtual void RecombinePin(UEdGraphPin* Pin) const override;
-	virtual void OnPinConnectionDoubleCicked(UEdGraphPin* PinA, UEdGraphPin* PinB, const FVector2D& GraphPosition) const override;
+	virtual void OnPinConnectionDoubleCicked(UEdGraphPin* PinA, UEdGraphPin* PinB, const FVector2f& GraphPosition) const override;
 	virtual UEdGraphPin* DropPinOnNode(UEdGraphNode* InTargetNode, const FName& InSourcePinName, const FEdGraphPinType& InSourcePinType, EEdGraphPinDirection InSourcePinDirection) const override;
 	virtual bool SupportsDropPinOnNode(UEdGraphNode* InTargetNode, const FEdGraphPinType& InSourcePinType, EEdGraphPinDirection InSourcePinDirection, FText& OutErrorMessage) const override;
 	virtual bool IsCacheVisualizationOutOfDate(int32 InVisualizationCacheID) const override;
diff --git a/Engine/Source/Editor/BlueprintGraph/Classes/EdGraphSchema_K2_Actions.h b/Engine/Source/Editor/BlueprintGraph/Classes/EdGraphSchema_K2_Actions.h
--- a/Engine/Source/Editor/BlueprintGraph/Classes/EdGraphSchema_K2_Actions.h
+++ b/Engine/Source/Editor/BlueprintGraph/Classes/EdGraphSchema_K2_Actions.h
@@ -83,8 +83,9 @@ struct BLUEPRINTGRAPH_API FEdGraphSchemaAction_K2NewNode : public FEdGraphSchema
 	{}
 
 	// FEdGraphSchemaAction interface
-	virtual UEdGraphNode* PerformAction(UEdGraph* ParentGraph, UEdGraphPin* FromPin, const FVector2D Location, bool bSelectNewNode = true) override;
-	virtual UEdGraphNode* PerformAction(UEdGraph* ParentGraph, TArray<UEdGraphPin*>& FromPins, const FVector2D Location, bool bSelectNewNode = true) override;
+	using FEdGraphSchemaAction::PerformAction; // Prevent hiding of deprecated base class function with FVector2D
+	virtual UEdGraphNode* PerformAction(UEdGraph* ParentGraph, UEdGraphPin* FromPin, const FVector2f& Location, bool bSelectNewNode = true) override;
+	virtual UEdGraphNode* PerformAction(UEdGraph* ParentGraph, TArray<UEdGraphPin*>& FromPins, const FVector2f& Location, bool bSelectNewNode = true) override;
 	virtual void AddReferencedObjects(FReferenceCollector& Collector) override;
 	// End of FEdGraphSchemaAction interface
 
@@ -157,8 +158,8 @@ struct BLUEPRINTGRAPH_API FEdGraphSchemaAction_K2ViewNode : public FEdGraphSchem
 	{}
 
 	// FEdGraphSchemaAction interface
-	virtual UEdGraphNode* PerformAction(class UEdGraph* ParentGraph, UEdGraphPin* FromPin, const FVector2D Location, bool bSelectNewNode = true) override;
-	virtual UEdGraphNode* PerformAction(class UEdGraph* ParentGraph, TArray<UEdGraphPin*>& FromPins, const FVector2D Location, bool bSelectNewNode = true) override;
+	virtual UEdGraphNode* PerformAction(class UEdGraph* ParentGraph, UEdGraphPin* FromPin, const FVector2f& Location, bool bSelectNewNode = true) override;
+	virtual UEdGraphNode* PerformAction(class UEdGraph* ParentGraph, TArray<UEdGraphPin*>& FromPins, const FVector2f& Location, bool bSelectNewNode = true) override;
 	// End of FEdGraphSchemaAction interface
 
 };
@@ -185,7 +186,7 @@ struct BLUEPRINTGRAPH_API FEdGraphSchemaAction_K2AssignDelegate : public FEdGrap
 	{}
 
 	// FEdGraphSchemaAction interface
-	virtual UEdGraphNode* PerformAction(class UEdGraph* ParentGraph, UEdGraphPin* FromPin, const FVector2D Location, bool bSelectNewNode = true) override;
+	virtual UEdGraphNode* PerformAction(class UEdGraph* ParentGraph, UEdGraphPin* FromPin, const FVector2f& Location, bool bSelectNewNode = true) override;
 	// End of FEdGraphSchemaAction interface
 
 	static UEdGraphNode* AssignDelegate(class UK2Node* NodeTemplate, class UEdGraph* ParentGraph, UEdGraphPin* FromPin, const FVector2D Location, bool bSelectNewNode);
@@ -218,8 +219,8 @@ struct BLUEPRINTGRAPH_API FEdGraphSchemaAction_EventFromFunction : public FEdGra
 	{}
 
 	// FEdGraphSchemaAction interface
-	virtual UEdGraphNode* PerformAction(class UEdGraph* ParentGraph, UEdGraphPin* FromPin, const FVector2D Location, bool bSelectNewNode = true) override;
-	virtual UEdGraphNode* PerformAction(class UEdGraph* ParentGraph, TArray<UEdGraphPin*>& FromPins, const FVector2D Location, bool bSelectNewNode = true) override;
+	virtual UEdGraphNode* PerformAction(class UEdGraph* ParentGraph, UEdGraphPin* FromPin, const FVector2f& Location, bool bSelectNewNode = true) override;
+	virtual UEdGraphNode* PerformAction(class UEdGraph* ParentGraph, TArray<UEdGraphPin*>& FromPins, const FVector2f& Location, bool bSelectNewNode = true) override;
 	virtual void AddReferencedObjects(FReferenceCollector& Collector) override;
 	// End of FEdGraphSchemaAction interface
 };
@@ -259,7 +260,7 @@ struct FEdGraphSchemaAction_K2AddComponent : public FEdGraphSchemaAction_K2NewNo
 	{}
 
 	// FEdGraphSchemaAction interface
-	virtual UEdGraphNode* PerformAction(class UEdGraph* ParentGraph, UEdGraphPin* FromPin, const FVector2D Location, bool bSelectNewNode = true) override;
+	virtual UEdGraphNode* PerformAction(class UEdGraph* ParentGraph, UEdGraphPin* FromPin, const FVector2f& Location, bool bSelectNewNode = true) override;
 	virtual void AddReferencedObjects(FReferenceCollector& Collector) override;
 	// End of FEdGraphSchemaAction interface
 };
@@ -287,7 +288,7 @@ struct FEdGraphSchemaAction_K2AddEvent : public FEdGraphSchemaAction_K2NewNode
 	{}
 
 	// FEdGraphSchemaAction interface
-	virtual UEdGraphNode* PerformAction(class UEdGraph* ParentGraph, UEdGraphPin* FromPin, const FVector2D Location, bool bSelectNewNode = true) override;
+	virtual UEdGraphNode* PerformAction(class UEdGraph* ParentGraph, UEdGraphPin* FromPin, const FVector2f& Location, bool bSelectNewNode = true) override;
 	// End of FEdGraphSchemaAction interface
 
 	/** */
@@ -317,7 +318,7 @@ struct FEdGraphSchemaAction_K2AddCustomEvent : public FEdGraphSchemaAction_K2New
 	{}
 
 	// FEdGraphSchemaAction interface
-	BLUEPRINTGRAPH_API virtual UEdGraphNode* PerformAction(class UEdGraph* ParentGraph, UEdGraphPin* FromPin, const FVector2D Location, bool bSelectNewNode = true) override;
+	BLUEPRINTGRAPH_API virtual UEdGraphNode* PerformAction(class UEdGraph* ParentGraph, UEdGraphPin* FromPin, const FVector2f& Location, bool bSelectNewNode = true) override;
 	// End of FEdGraphSchemaAction interface
 };
 
@@ -348,7 +349,7 @@ struct FEdGraphSchemaAction_K2AddCallOnActor : public FEdGraphSchemaAction_K2New
 	{}
 
 	// FEdGraphSchemaAction interface
-	virtual UEdGraphNode* PerformAction(class UEdGraph* ParentGraph, UEdGraphPin* FromPin, const FVector2D Location, bool bSelectNewNode = true) override;
+	virtual UEdGraphNode* PerformAction(class UEdGraph* ParentGraph, UEdGraphPin* FromPin, const FVector2f& Location, bool bSelectNewNode = true) override;
 	virtual void AddReferencedObjects(FReferenceCollector& Collector) override;
 	// End of FEdGraphSchemaAction interface	
 };
@@ -378,7 +379,8 @@ struct FEdGraphSchemaAction_K2AddComment : public FEdGraphSchemaAction
 	}
 
 	// FEdGraphSchemaAction interface
-	BLUEPRINTGRAPH_API virtual UEdGraphNode* PerformAction(class UEdGraph* ParentGraph, UEdGraphPin* FromPin, const FVector2D Location, bool bSelectNewNode = true) override;
+	using FEdGraphSchemaAction::PerformAction; // Prevent hiding of deprecated base class function with FVector2D
+	BLUEPRINTGRAPH_API virtual UEdGraphNode* PerformAction(class UEdGraph* ParentGraph, UEdGraphPin* FromPin, const FVector2f& Location, bool bSelectNewNode = true) override;
 	// End of FEdGraphSchemaAction interface
 };
 
@@ -406,7 +408,8 @@ struct BLUEPRINTGRAPH_API FEdGraphSchemaAction_K2TargetNode : public FEdGraphSch
 	{}
 
 	// FEdGraphSchemaAction interface
-	virtual UEdGraphNode* PerformAction(class UEdGraph* ParentGraph, UEdGraphPin* FromPin, const FVector2D Location, bool bSelectNewNode = true) override;
+	using FEdGraphSchemaAction::PerformAction; // Prevent hiding of deprecated base class function with FVector2D
+	virtual UEdGraphNode* PerformAction(class UEdGraph* ParentGraph, UEdGraphPin* FromPin, const FVector2f& Location, bool bSelectNewNode = true) override;
 	// End of FEdGraphSchemaAction interface
 };
 
@@ -433,7 +436,8 @@ struct BLUEPRINTGRAPH_API FEdGraphSchemaAction_K2PasteHere : public FEdGraphSche
 	{}
 
 	// FEdGraphSchemaAction interface
-	virtual UEdGraphNode* PerformAction(class UEdGraph* ParentGraph, UEdGraphPin* FromPin, const FVector2D Location, bool bSelectNewNode = true) override;
+	using FEdGraphSchemaAction::PerformAction; // Prevent hiding of deprecated base class function with FVector2D
+	virtual UEdGraphNode* PerformAction(class UEdGraph* ParentGraph, UEdGraphPin* FromPin, const FVector2f& Location, bool bSelectNewNode = true) override;
 	// End of FEdGraphSchemaAction interface
 };
 
diff --git a/Engine/Source/Editor/BlueprintGraph/Private/EdGraphSchema_K2.cpp b/Engine/Source/Editor/BlueprintGraph/Private/EdGraphSchema_K2.cpp
--- a/Engine/Source/Editor/BlueprintGraph/Private/EdGraphSchema_K2.cpp
+++ b/Engine/Source/Editor/BlueprintGraph/Private/EdGraphSchema_K2.cpp
@@ -5841,7 +5841,7 @@ bool UEdGraphSchema_K2::IsStaticFunctionGraph( const UEdGraph* TestEdGraph ) con
 	return false;
 }
 
-void UEdGraphSchema_K2::DroppedAssetsOnGraph(const TArray<FAssetData>& Assets, const FVector2D& GraphPosition, UEdGraph* Graph) const 
+void UEdGraphSchema_K2::DroppedAssetsOnGraph(const TArray<FAssetData>& Assets, const FVector2f& GraphPosition, UEdGraph* Graph) const 
 {
 	// only want to spawn event nodes in an event graph
 	if (FBlueprintEditorUtils::IsEventGraph(Graph))
@@ -5853,7 +5853,7 @@ void UEdGraphSchema_K2::DroppedAssetsOnGraph(const TArray<FAssetData>& Assets, c
 			// if we can find any actions we want to try to spawn the node - only spawning Input Action event nodes currently
 			if (const FAssetBlueprintGraphActions* GraphActions = Module.GetAssetBlueprintGraphActions(AssetData.GetClass()))
 			{
-				GraphActions->TryCreatingAssetNode(AssetData, Graph, GraphPosition, EK2NewNodeFlags::SelectNewNode);
+				GraphActions->TryCreatingAssetNode(AssetData, Graph, FDeprecateSlateVector2D(GraphPosition), EK2NewNodeFlags::SelectNewNode);
 			}
 		}
 	}
@@ -5864,7 +5864,7 @@ void UEdGraphSchema_K2::DroppedAssetsOnGraph(const TArray<FAssetData>& Assets, c
 		float XOffset = 0.0f;
 		for(int32 AssetIdx=0; AssetIdx < Assets.Num(); AssetIdx++)
 		{
-			FVector2D Position = GraphPosition + (AssetIdx * FVector2D(XOffset, 0.0f));
+			FVector2f Position = GraphPosition + (AssetIdx * FVector2f(XOffset, 0.0f));
 
 			UObject* Asset = Assets[AssetIdx].GetAsset();
 			if (!Asset)
@@ -5905,7 +5905,7 @@ void UEdGraphSchema_K2::DroppedAssetsOnGraph(const TArray<FAssetData>& Assets, c
 				{
 					IBlueprintNodeBinder::FBindingSet Bindings;
 					Bindings.Add(Asset);
-					Spawner->Invoke(Graph, Bindings, GraphPosition);
+					Spawner->Invoke(Graph, Bindings, FDeprecateSlateVector2D(GraphPosition));
 				}
 				else
 				{
@@ -5916,12 +5916,12 @@ void UEdGraphSchema_K2::DroppedAssetsOnGraph(const TArray<FAssetData>& Assets, c
 	}
 }
 
-void UEdGraphSchema_K2::DroppedAssetsOnNode(const TArray<FAssetData>& Assets, const FVector2D& GraphPosition, UEdGraphNode* Node) const
+void UEdGraphSchema_K2::DroppedAssetsOnNode(const TArray<FAssetData>& Assets, const FVector2f& GraphPosition, UEdGraphNode* Node) const
 {
 	// @TODO: Should dropping on component node change the component?
 }
 
-void UEdGraphSchema_K2::DroppedAssetsOnPin(const TArray<FAssetData>& Assets, const FVector2D& GraphPosition, UEdGraphPin* Pin) const
+void UEdGraphSchema_K2::DroppedAssetsOnPin(const TArray<FAssetData>& Assets, const FVector2f& GraphPosition, UEdGraphPin* Pin) const
 {
 	// If dropping onto an 'object' pin, try and set the literal
 	if ((Pin->PinType.PinCategory == PC_Object) || (Pin->PinType.PinCategory == PC_Interface))
@@ -7490,19 +7490,19 @@ void UEdGraphSchema_K2::RecombinePin(UEdGraphPin* Pin) const
 	FBlueprintEditorUtils::MarkBlueprintAsModified(Blueprint);
 }
 
-void UEdGraphSchema_K2::OnPinConnectionDoubleCicked(UEdGraphPin* PinA, UEdGraphPin* PinB, const FVector2D& GraphPosition) const
+void UEdGraphSchema_K2::OnPinConnectionDoubleCicked(UEdGraphPin* PinA, UEdGraphPin* PinB, const FVector2f& GraphPosition) const
 {
 	const FScopedTransaction Transaction(LOCTEXT("CreateRerouteNodeOnWire", "Create Reroute Node"), UEdGraphSchemaImpl::ShouldActuallyTransact());
 
 	//@TODO: This constant is duplicated from inside of SGraphNodeKnot
-	const FVector2D NodeSpacerSize(42.0f, 24.0f);
-	const FVector2D KnotTopLeft = GraphPosition - (NodeSpacerSize * 0.5f);
+	const FVector2f NodeSpacerSize(42.0f, 24.0f);
+	const FVector2f KnotTopLeft = GraphPosition - (NodeSpacerSize * 0.5f);
 
 	// Create a new knot
 	UEdGraph* ParentGraph = PinA->GetOwningNode()->GetGraph();
 	if (!FBlueprintEditorUtils::IsGraphReadOnly(ParentGraph))
 	{
-		UK2Node_Knot* NewKnot = FEdGraphSchemaAction_K2NewNode::SpawnNode<UK2Node_Knot>(ParentGraph, KnotTopLeft, EK2NewNodeFlags::SelectNewNode);
+		UK2Node_Knot* NewKnot = FEdGraphSchemaAction_K2NewNode::SpawnNode<UK2Node_Knot>(ParentGraph, FDeprecateSlateVector2D(KnotTopLeft), EK2NewNodeFlags::SelectNewNode);
 
 		// Move the connections across (only notifying the knot, as the other two didn't really change)
 		PinA->BreakLinkTo(PinB);
