Reconnect pins on nodes that have been replaced before recreating orphaned pins

[RN] Blueprint nodes that are replaced on load can now reconnect mistakenly orphaned pins

diff --git a/Engine/Source/Editor/BlueprintGraph/Private/EdGraphSchema_K2.cpp b/Engine/Source/Editor/BlueprintGraph/Private/EdGraphSchema_K2.cpp
--- a/Engine/Source/Editor/BlueprintGraph/Private/EdGraphSchema_K2.cpp
+++ b/Engine/Source/Editor/BlueprintGraph/Private/EdGraphSchema_K2.cpp
@@ -6096,18 +6096,66 @@ struct FBackwardCompatibilityConversionHelper
 
 		if (OrphanedLinks.Num() > 0)
 		{
-			// Add an orphan pin so warning/connections are not silently lost
-			UEdGraphPin* OrphanPin = NewNode->CreatePin(OldPin->Direction, OldPin->PinType, OldPin->PinName);
+			// The link may have been to a pin on a node that was *also* recreated, so now that
+			// we're recreated, lets make one last ditch effort to reconnect the link:
+			const auto IsUnorphanedPinByNameAndDirection = [](const UEdGraphPin* PinInQuestion, const UEdGraphPin* PinToMatch) 
+			{ 
+				return !PinInQuestion->bOrphanedPin && 
+					PinInQuestion->PinName == PinToMatch->PinName && 
+					PinInQuestion->Direction == PinToMatch->Direction;  
+			};
 
-			OrphanPin->bOrphanedPin = true;
-			OrphanPin->bNotConnectable = true;
+			const UEdGraphSchema_K2* Schema = GetDefault<UEdGraphSchema_K2>();
+			TArray<TPair<UEdGraphPin*, UEdGraphPin*>> PinsToConnect;
+			TArray<UEdGraphPin*> RelinkedOrphanedLinks;
+			for(UEdGraphPin* LinkedPin : OrphanedLinks)
+			{
+				if(!LinkedPin->bOrphanedPin)
+				{
+					continue;
+				}
 
-			for (UEdGraphPin* OldLink : OrphanedLinks)
+				UEdGraphNode* Node = LinkedPin->GetOwningNode();
+				UEdGraphPin* NewPinToLink = Node->FindPinByPredicate(
+					[&IsUnorphanedPinByNameAndDirection, LinkedPin]
+					(const UEdGraphPin* NodePin) 
+					{
+						return IsUnorphanedPinByNameAndDirection(NodePin, LinkedPin);
+					}
+				);
+
+				if( NewPinToLink &&
+					Schema->TryCreateConnection(NewPinToLink, NewPin) )
+				{
+					RelinkedOrphanedLinks.AddUnique(LinkedPin); // NewPinToLink is now linked to NewPin
+				}
+			}
+			
+			// remove pins with no connections if we successfully grabbed all of their connections:
+			for(UEdGraphPin* DirtyPin : RelinkedOrphanedLinks)
 			{
-				OrphanPin->MakeLinkTo(OldLink);
+				if(DirtyPin->LinkedTo.Num() == 0)
+				{
+					DirtyPin->GetOwningNode()->RemovePin(DirtyPin);
+					OrphanedLinks.Remove(DirtyPin);
+				}
 			}
+			
+			if (OrphanedLinks.Num() > 0)
+			{
+				// reconnection failed so add an orphan pin so warning/connections are not silently lost:
+				UEdGraphPin* OrphanPin = NewNode->CreatePin(OldPin->Direction, OldPin->PinType, OldPin->PinName);
 
-			return true;
+				OrphanPin->bOrphanedPin = true;
+				OrphanPin->bNotConnectable = true;
+
+				for (UEdGraphPin* OldLink : OrphanedLinks)
+				{
+					OrphanPin->MakeLinkTo(OldLink);
+				}
+
+				return true;
+			}
 		}
 
 		return false;
