PR #13090: Expand 'HideAssetPicker' meta usefulness

This expands the metadata to work on other node types besides function calls (eg: property setters).

Additionally, the metadata was renamed to "HidePinAssetPicker". Previously, it was a little ambiguous because it could give the incorrect impression that it hides the picker on the details panel for any properties.

diff --git a/Engine/Source/Editor/BlueprintGraph/Classes/EdGraphSchema_K2.h b/Engine/Source/Editor/BlueprintGraph/Classes/EdGraphSchema_K2.h
--- a/Engine/Source/Editor/BlueprintGraph/Classes/EdGraphSchema_K2.h
+++ b/Engine/Source/Editor/BlueprintGraph/Classes/EdGraphSchema_K2.h
@@ -194,9 +194,13 @@ public:
 	/** If true, an unconnected pin will generate a UPROPERTY under the hood to connect as the input, which will be set to the literal value for the pin.  Only valid for reference parameters. */
 	static const FName MD_AutoCreateRefTerm;
 
-	/** The specified parameter should hide the asset picker on the pin, even if it is a valid UObject Asset. */
+	UE_DEPRECATED(5.6, "Use MD_HidePinAssetPicker instead.")
 	static const FName MD_HideAssetPicker;
 
+	/** For functions, the specified parameter should hide the asset picker on the pin, even if it is a valid UObject Asset. */
+	/** For properties, forcibly hide the asset picker for the pin corresponding to the property. */
+	static const FName MD_HidePinAssetPicker;
+
 	/** If true, the hidden world context pin will be visible when the function is placed in a child blueprint of the class. */
 	static const FName MD_ShowWorldContextPin;
 
diff --git a/Engine/Source/Editor/BlueprintGraph/Private/EdGraphSchema_K2.cpp b/Engine/Source/Editor/BlueprintGraph/Private/EdGraphSchema_K2.cpp
--- a/Engine/Source/Editor/BlueprintGraph/Private/EdGraphSchema_K2.cpp
+++ b/Engine/Source/Editor/BlueprintGraph/Private/EdGraphSchema_K2.cpp
@@ -204,6 +204,7 @@ const FName FBlueprintMetadata::MD_CallableWithoutWorldContext(TEXT("CallableWit
 const FName FBlueprintMetadata::MD_DevelopmentOnly(TEXT("DevelopmentOnly"));
 const FName FBlueprintMetadata::MD_AutoCreateRefTerm(TEXT("AutoCreateRefTerm"));
 const FName FBlueprintMetadata::MD_HideAssetPicker(TEXT("HideAssetPicker"));
+const FName FBlueprintMetadata::MD_HidePinAssetPicker(TEXT("HidePinAssetPicker"));
 
 const FName FBlueprintMetadata::MD_ShowWorldContextPin(TEXT("ShowWorldContextPin"));
 const FName FBlueprintMetadata::MD_Private(TEXT("BlueprintPrivate"));
@@ -5135,19 +5136,30 @@ bool UEdGraphSchema_K2::ShouldShowAssetPickerForPin(UEdGraphPin* Pin) const
 					if (UFunction* FunctionRef = CallFunctionNode->GetTargetFunction())
 					{
 						const UEdGraphPin* WorldContextPin = CallFunctionNode->FindPin(FunctionRef->GetMetaData(FBlueprintMetadata::MD_WorldContext));
-						bShow = ( WorldContextPin != Pin );
+						bShow = (WorldContextPin != Pin);
 
 						// Check if we have explictly marked this pin as hiding the asset picker
-						const FString& HideAssetPickerMetaData = FunctionRef->GetMetaData(FBlueprintMetadata::MD_HideAssetPicker);
-						if(!HideAssetPickerMetaData.IsEmpty())
+						FString HideAssetPickerMetaData = FunctionRef->GetMetaData(FBlueprintMetadata::MD_HidePinAssetPicker);
+						if (HideAssetPickerMetaData.IsEmpty())
+						{
+							PRAGMA_DISABLE_DEPRECATION_WARNINGS
+							HideAssetPickerMetaData = FunctionRef->GetMetaData(FBlueprintMetadata::MD_HideAssetPicker);
+							PRAGMA_ENABLE_DEPRECATION_WARNINGS
+							if (!HideAssetPickerMetaData.IsEmpty())
+							{
+								UE_LOG(LogBlueprint, Warning, TEXT("Use of deprecated meta 'HideAssetPicker' for function %s. Change to HidePinAssetPicker."), *FunctionRef->GetPathName());
+							}
+						}
+
+						if (!HideAssetPickerMetaData.IsEmpty())
 						{
 							TArray<FString> PinNames;
 							HideAssetPickerMetaData.ParseIntoArray(PinNames, TEXT(","), true);
 							const FString PinName = Pin->GetName();
-							for(FString& ParamNameToHide : PinNames)
+							for (FString& ParamNameToHide : PinNames)
 							{
 								ParamNameToHide.TrimStartAndEndInline();
-								if(ParamNameToHide == PinName)
+								if (ParamNameToHide == PinName)
 								{
 									bShow = false;
 									break;
@@ -5156,10 +5168,16 @@ bool UEdGraphSchema_K2::ShouldShowAssetPickerForPin(UEdGraphPin* Pin) const
 						}
 					}
 				}
-				else if (Cast<UK2Node_CreateDelegate>( Pin->GetOwningNode())) 
+				else if (Cast<UK2Node_CreateDelegate>(Pin->GetOwningNode())) 
 				{
 					bShow = false;
 				}
+
+				const FString& Meta = bShow ? Pin->GetOwningNode()->GetPinMetaData(Pin->PinName, FBlueprintMetadata::MD_HidePinAssetPicker) : FString();
+				if (!Meta.IsEmpty())
+				{
+					bShow = !(Meta.Compare("true", ESearchCase::IgnoreCase) == 0);
+				}
 			}
 		}
 	}
diff --git a/Engine/Source/Runtime/CoreUObject/Public/UObject/ObjectMacros.h b/Engine/Source/Runtime/CoreUObject/Public/UObject/ObjectMacros.h
--- a/Engine/Source/Runtime/CoreUObject/Public/UObject/ObjectMacros.h
+++ b/Engine/Source/Runtime/CoreUObject/Public/UObject/ObjectMacros.h
@@ -1545,6 +1545,9 @@ namespace UM
 
 		/// [PropertyMetadata] Used for edit inline UObject* properties, enables detail class customization for instanced UObject properties.
 		AllowEditInlineCustomization,
+
+		/// [PropertyMetadata] Forcibly hide the asset picker for the pin corresponding to the property.
+		// HidePinAssetPicker, (Commented out so as to avoid duplicate name with version in the Function section, but still show in the property section)
 	};
 
 	// Metadata usable in UPROPERTY for customizing the behavior of Persona and UMG
@@ -1593,8 +1596,10 @@ namespace UM
 		/// This also allows for inline editing of the default value on some types (take FRotator for instance). Only valid for inputs.
 		AutoCreateRefTerm,
 
+		HideAssetPicker UE_DEPRECATED(5.6, "Use HidePinAssetPicker instead."),
+
 		/// [FunctionMetadata] Forcibly hide the asset picker for pins matching any parameter names in this piece of metadata.
-		HideAssetPicker,
+		HidePinAssetPicker,
 
 		/// [FunctionMetadata] This function is an internal implementation detail, used to implement another function or node.  It is never directly exposed in a graph.
 		BlueprintInternalUseOnly,
