PR #11622: Implement ForEachLoop node for TSet datatypes.

Additional changes made by Epic:

* Used minimal API to reduce export count.
* Added value name customization.
* Changed PinConnectionListChanged to NotifyPinConnectionListChanged.
* Replaced type state-tracking with RefreshWildcardPins, which updates the pins based on the current set link. This also simplifies the implementation.
* Removed pin tooltip customization.
* Fixed an iterator bug in GenericSet_GetItemByIndex.

diff --git a/Engine/Source/Editor/BlueprintGraph/Classes/K2Node_SetForEach.h b/Engine/Source/Editor/BlueprintGraph/Classes/K2Node_SetForEach.h
new file mode 100644
--- /dev/null
+++ b/Engine/Source/Editor/BlueprintGraph/Classes/K2Node_SetForEach.h
@@ -0,0 +1,64 @@
+// Copyright Epic Games, Inc. All Rights Reserved.
+
+#pragma once
+
+#include "K2Node.h"
+
+#include "K2Node_SetForEach.generated.h"
+
+/** Custom blueprint node for iterating Sets in blueprints */
+UCLASS(MinimalAPI, CollapseCategories)
+class UK2Node_SetForEach : public UK2Node
+{
+	GENERATED_BODY()
+
+public:
+
+	UK2Node_SetForEach();
+
+	/** Pin Accessors */
+	[[nodiscard]] UEdGraphPin* GetSetPin() const;
+	[[nodiscard]] UEdGraphPin* GetBreakPin() const;
+	[[nodiscard]] UEdGraphPin* GetForEachPin() const;
+	[[nodiscard]] UEdGraphPin* GetValuePin() const;
+	[[nodiscard]] UEdGraphPin* GetCompletedPin() const;
+
+	/** UK2Node interface */
+	[[nodiscard]] virtual bool IsNodeSafeToIgnore() const override { return true; }
+	virtual void GetMenuActions(FBlueprintActionDatabaseRegistrar& ActionRegistrar) const override;
+	[[nodiscard]] virtual FText GetMenuCategory() const override;
+	virtual void PostReconstructNode() override;
+	/** End of UK2Node interface */
+
+	/** UEdGraphNode interface */
+	virtual void AllocateDefaultPins() override;
+	virtual void ExpandNode(FKismetCompilerContext& CompilerContext, UEdGraph* SourceGraph) override;
+	[[nodiscard]] virtual FText GetNodeTitle(ENodeTitleType::Type TitleType) const override;
+	[[nodiscard]] virtual FText GetTooltipText() const override;
+	[[nodiscard]] virtual FSlateIcon GetIconAndTint(FLinearColor& OutColor) const override;
+	virtual void NotifyPinConnectionListChanged(UEdGraphPin* Pin) override;
+	[[nodiscard]] virtual bool ShouldShowNodeProperties() const override { return true; }
+	/** End of UEdGraphNode interface */
+
+	/** UObject interface */
+	void PostEditChangeProperty(FPropertyChangedEvent& PropertyChangedEvent) override;
+	/** End of UObject interface */
+
+private:
+
+	/** Pin Names */
+	static const FName SetPinName;
+	static const FName BreakPinName;
+	static const FName ValuePinName;
+	static const FName CompletedPinName;
+
+	/** Determine if there is any configuration options that shouldn't be allowed */
+	[[nodiscard]] bool CheckForErrors(const FKismetCompilerContext& CompilerContext);
+
+	/** Updates the wildcard pins based on current links */
+	void RefreshWildcardPins();
+
+	/** A user editable hook for the display name of the value pin */
+	UPROPERTY(EditDefaultsOnly, Category = "Set For Each")
+	FString ValueName;
+};
diff --git a/Engine/Source/Editor/BlueprintGraph/Private/K2Node_SetForEach.cpp b/Engine/Source/Editor/BlueprintGraph/Private/K2Node_SetForEach.cpp
new file mode 100644
--- /dev/null
+++ b/Engine/Source/Editor/BlueprintGraph/Private/K2Node_SetForEach.cpp
@@ -0,0 +1,383 @@
+// Copyright Epic Games, Inc. All Rights Reserved.
+
+#include "K2Node_SetForEach.h"
+
+#include "BlueprintActionDatabaseRegistrar.h"
+#include "BlueprintNodeSpawner.h"
+#include "K2Node_AssignmentStatement.h"
+#include "K2Node_CallFunction.h"
+#include "K2Node_ExecutionSequence.h"
+#include "K2Node_IfThenElse.h"
+#include "K2Node_TemporaryVariable.h"
+#include "KismetCompiler.h"
+#include "Kismet/BlueprintSetLibrary.h"
+#include "Kismet/KismetMathLibrary.h"
+#include "Kismet2/BlueprintEditorUtils.h"
+#include "Kismet2/WildcardNodeUtils.h"
+
+#define LOCTEXT_NAMESPACE "K2Node_SetForEach"
+
+const FName UK2Node_SetForEach::SetPinName(TEXT("SetPin"));
+const FName UK2Node_SetForEach::BreakPinName(TEXT("BreakPin"));
+const FName UK2Node_SetForEach::ValuePinName(TEXT("ValuePin"));
+const FName UK2Node_SetForEach::CompletedPinName(TEXT("CompletedPin"));
+
+UK2Node_SetForEach::UK2Node_SetForEach()
+{
+	ValueName = LOCTEXT("ValuePin_FriendlyName", "Set Value").ToString();
+}
+
+UEdGraphPin* UK2Node_SetForEach::GetSetPin() const
+{
+	return FindPinChecked(SetPinName);
+}
+
+UEdGraphPin* UK2Node_SetForEach::GetBreakPin() const
+{
+	return FindPinChecked(BreakPinName);
+}
+
+UEdGraphPin* UK2Node_SetForEach::GetForEachPin() const
+{
+	return FindPinChecked(UEdGraphSchema_K2::PN_Then);
+}
+
+UEdGraphPin* UK2Node_SetForEach::GetValuePin() const
+{
+	return FindPinChecked(ValuePinName);
+}
+
+UEdGraphPin* UK2Node_SetForEach::GetCompletedPin() const
+{
+	return FindPinChecked(CompletedPinName);
+}
+
+void UK2Node_SetForEach::GetMenuActions(FBlueprintActionDatabaseRegistrar& ActionRegistrar) const
+{
+	// Actions are registered under specific object-keys; the idea is that 
+	// actions might have to be updated (or deleted) if their object-key is  
+	// mutated (or removed).
+	// Here, we use the node's class to keep from needlessly instantiating a UBlueprintNodeSpawner.
+	// Additionally, if the node type disappears, then the action should go with it.
+
+	UClass* ActionKey = GetClass();
+	if (ActionRegistrar.IsOpenForRegistration(ActionKey))
+	{
+		UBlueprintNodeSpawner* NodeSpawner = UBlueprintNodeSpawner::Create(ActionKey);
+		check(NodeSpawner != nullptr);
+		ActionRegistrar.AddBlueprintAction(ActionKey, NodeSpawner);
+	}
+}
+
+FText UK2Node_SetForEach::GetMenuCategory() const
+{
+	return LOCTEXT("NodeMenu", "Utilities|Set");
+}
+
+void UK2Node_SetForEach::PostReconstructNode()
+{
+	Super::PostReconstructNode();
+
+	RefreshWildcardPins();
+}
+
+void UK2Node_SetForEach::AllocateDefaultPins()
+{
+	Super::AllocateDefaultPins();
+
+	CreatePin(EGPD_Input, UEdGraphSchema_K2::PC_Exec, UEdGraphSchema_K2::PN_Execute);
+
+	UEdGraphNode::FCreatePinParams PinParams;
+	PinParams.ContainerType = EPinContainerType::Set;
+	PinParams.ValueTerminalType.TerminalCategory = UEdGraphSchema_K2::PC_Wildcard;
+
+	UEdGraphPin* SetPin = CreatePin(EGPD_Input, UEdGraphSchema_K2::PC_Wildcard, SetPinName, PinParams);
+	SetPin->PinType.bIsConst = true;
+	SetPin->PinType.bIsReference = true;
+	SetPin->PinFriendlyName = LOCTEXT("SetPin_FriendlyName", "Set");
+
+	UEdGraphPin* BreakPin = CreatePin(EGPD_Input, UEdGraphSchema_K2::PC_Exec, BreakPinName);
+	BreakPin->PinFriendlyName = LOCTEXT("BreakPin_FriendlyName", "Break");
+	BreakPin->bAdvancedView = true;
+
+	UEdGraphPin* ForEachPin = CreatePin(EGPD_Output, UEdGraphSchema_K2::PC_Exec, UEdGraphSchema_K2::PN_Then);
+	ForEachPin->PinFriendlyName = LOCTEXT("ForEachPin_FriendlyName", "Loop Body");
+
+	UEdGraphPin* ValuePin = CreatePin(EGPD_Output, UEdGraphSchema_K2::PC_Wildcard, ValuePinName);
+	ValuePin->PinFriendlyName = FText::FromString(ValueName);
+
+	UEdGraphPin* CompletedPin = CreatePin(EGPD_Output, UEdGraphSchema_K2::PC_Exec, CompletedPinName);
+	CompletedPin->PinFriendlyName = LOCTEXT("CompletedPin_FriendlyName", "Completed");
+	CompletedPin->PinToolTip = LOCTEXT("CompletedPin_Tooltip", "Execution once all set elements have been visited").ToString();
+
+	if (AdvancedPinDisplay == ENodeAdvancedPins::NoPins)
+	{
+		AdvancedPinDisplay = ENodeAdvancedPins::Hidden;
+	}
+}
+
+void UK2Node_SetForEach::ExpandNode(FKismetCompilerContext& CompilerContext, UEdGraph* SourceGraph)
+{
+	Super::ExpandNode(CompilerContext, SourceGraph);
+
+	if (CheckForErrors(CompilerContext))
+	{
+		// Remove all the links to this node as they are no longer needed
+		BreakAllNodeLinks();
+		return;
+	}
+
+	const UEdGraphSchema_K2* K2Schema = CastChecked<UEdGraphSchema_K2>(GetSchema());
+
+	///////////////////////////////////////////////////////////////////////////////////
+	// Cache off versions of all our important pins
+
+	UEdGraphPin* ForEach_Exec = GetExecPin();
+	UEdGraphPin* ForEach_Set = GetSetPin();
+	UEdGraphPin* ForEach_Break = GetBreakPin();
+	UEdGraphPin* ForEach_ForEach = GetForEachPin();
+	UEdGraphPin* ForEach_Value = GetValuePin();
+	UEdGraphPin* ForEach_Completed = GetCompletedPin();
+
+	///////////////////////////////////////////////////////////////////////////////////
+	// Create a loop counter variable
+
+	UK2Node_TemporaryVariable* CreateTemporaryVariable = CompilerContext.SpawnIntermediateNode<UK2Node_TemporaryVariable>(this, SourceGraph);
+	CreateTemporaryVariable->VariableType.PinCategory = UEdGraphSchema_K2::PC_Int;
+	CreateTemporaryVariable->AllocateDefaultPins();
+
+	UEdGraphPin* Temp_Variable = CreateTemporaryVariable->GetVariablePin();
+
+	///////////////////////////////////////////////////////////////////////////////////
+	// Initialize the temporary to 0
+
+	UK2Node_AssignmentStatement* InitTemporaryVariable = CompilerContext.SpawnIntermediateNode<UK2Node_AssignmentStatement>(this, SourceGraph);
+	InitTemporaryVariable->AllocateDefaultPins();
+
+	UEdGraphPin* Init_Exec = InitTemporaryVariable->GetExecPin();
+	UEdGraphPin* Init_Variable = InitTemporaryVariable->GetVariablePin();
+	UEdGraphPin* Init_Value = InitTemporaryVariable->GetValuePin();
+	UEdGraphPin* Init_Then = InitTemporaryVariable->GetThenPin();
+
+	CompilerContext.MovePinLinksToIntermediate(*ForEach_Exec, *Init_Exec);
+	K2Schema->TryCreateConnection(Init_Variable, Temp_Variable);
+	Init_Value->DefaultValue = TEXT("0");
+
+	///////////////////////////////////////////////////////////////////////////////////
+	// Branch on comparing the loop index with the size of the set
+
+	UK2Node_IfThenElse* BranchOnIndex = CompilerContext.SpawnIntermediateNode<UK2Node_IfThenElse>(this, SourceGraph);
+	BranchOnIndex->AllocateDefaultPins();
+
+	UEdGraphPin* Branch_Exec = BranchOnIndex->GetExecPin();
+	UEdGraphPin* Branch_Input = BranchOnIndex->GetConditionPin();
+	UEdGraphPin* Branch_Then = BranchOnIndex->GetThenPin();
+	UEdGraphPin* Branch_Else = BranchOnIndex->GetElsePin();
+
+	Init_Then->MakeLinkTo(Branch_Exec);
+	CompilerContext.MovePinLinksToIntermediate(*ForEach_Completed, *Branch_Else);
+
+	UK2Node_CallFunction* CompareLessThan = CompilerContext.SpawnIntermediateNode<UK2Node_CallFunction>(this, SourceGraph);
+	CompareLessThan->FunctionReference.SetExternalMember( GET_FUNCTION_NAME_CHECKED(UKismetMathLibrary, Less_IntInt), UKismetMathLibrary::StaticClass());
+	CompareLessThan->AllocateDefaultPins();
+
+	UEdGraphPin* Compare_A = CompareLessThan->FindPinChecked(TEXT("A"));
+	UEdGraphPin* Compare_B = CompareLessThan->FindPinChecked(TEXT("B"));
+	UEdGraphPin* Compare_Return = CompareLessThan->GetReturnValuePin();
+
+	Branch_Input->MakeLinkTo(Compare_Return);
+	Temp_Variable->MakeLinkTo(Compare_A);
+
+	UK2Node_CallFunction* GetSetLength = CompilerContext.SpawnIntermediateNode<UK2Node_CallFunction>(this, SourceGraph);
+	GetSetLength->FunctionReference.SetExternalMember(GET_FUNCTION_NAME_CHECKED(UBlueprintSetLibrary, Set_Length), UBlueprintSetLibrary::StaticClass());
+	GetSetLength->AllocateDefaultPins();
+
+	UEdGraphPin* SetLength_Set = GetSetLength->FindPinChecked(TEXT("TargetSet"));
+	UEdGraphPin* SetLength_Return = GetSetLength->GetReturnValuePin();
+
+	// Coerce the wildcard pin types
+	SetLength_Set->PinType = ForEach_Set->PinType;
+
+	Compare_B->MakeLinkTo(SetLength_Return);
+	CompilerContext.CopyPinLinksToIntermediate(*ForEach_Set, *SetLength_Set);
+
+	///////////////////////////////////////////////////////////////////////////////////
+	// Sequence the loop body and incrementing the loop counter
+
+	UK2Node_ExecutionSequence* LoopSequence = CompilerContext.SpawnIntermediateNode<UK2Node_ExecutionSequence>(this, SourceGraph);
+	LoopSequence->AllocateDefaultPins();
+
+	UEdGraphPin* Sequence_Exec = LoopSequence->GetExecPin();
+	UEdGraphPin* Sequence_One = LoopSequence->GetThenPinGivenIndex(0);
+	UEdGraphPin* Sequence_Two = LoopSequence->GetThenPinGivenIndex(1);
+
+	Branch_Then->MakeLinkTo(Sequence_Exec);
+	CompilerContext.MovePinLinksToIntermediate(*ForEach_ForEach, *Sequence_One);
+
+	UK2Node_CallFunction* GetSetElement = CompilerContext.SpawnIntermediateNode<UK2Node_CallFunction>(this, SourceGraph);
+	GetSetElement->FunctionReference.SetExternalMember(GET_FUNCTION_NAME_CHECKED(UBlueprintSetLibrary, Set_GetItemByIndex), UBlueprintSetLibrary::StaticClass());
+	GetSetElement->AllocateDefaultPins();
+
+	UEdGraphPin* GetElement_Set = GetSetElement->FindPinChecked(TEXT("TargetSet"));
+	UEdGraphPin* GetElement_Index = GetSetElement->FindPinChecked(TEXT("Index"));
+	UEdGraphPin* GetElement_Value = GetSetElement->FindPinChecked(TEXT("Item"));
+
+	// Coerce the wildcard pin types
+	GetElement_Set->PinType = ForEach_Set->PinType;
+	GetElement_Value->PinType = ForEach_Value->PinType;
+
+	CompilerContext.CopyPinLinksToIntermediate(*ForEach_Set, *GetElement_Set);
+	GetElement_Index->MakeLinkTo(Temp_Variable);
+	CompilerContext.MovePinLinksToIntermediate(*ForEach_Value, *GetElement_Value);
+
+	///////////////////////////////////////////////////////////////////////////////////
+	// Increment the loop counter by one
+
+	const auto IncrementVariable = CompilerContext.SpawnIntermediateNode<UK2Node_AssignmentStatement>(this, SourceGraph);
+	IncrementVariable->AllocateDefaultPins();
+
+	UEdGraphPin* Inc_Exec = IncrementVariable->GetExecPin();
+	UEdGraphPin* Inc_Variable = IncrementVariable->GetVariablePin();
+	UEdGraphPin* Inc_Value = IncrementVariable->GetValuePin();
+	UEdGraphPin* Inc_Then = IncrementVariable->GetThenPin();
+
+	Sequence_Two->MakeLinkTo(Inc_Exec);
+	Branch_Exec->MakeLinkTo(Inc_Then);
+	K2Schema->TryCreateConnection(Temp_Variable, Inc_Variable);
+
+	UK2Node_CallFunction* AddOne = CompilerContext.SpawnIntermediateNode<UK2Node_CallFunction>(this, SourceGraph);
+	AddOne->FunctionReference.SetExternalMember(GET_FUNCTION_NAME_CHECKED(UKismetMathLibrary, Add_IntInt), UKismetMathLibrary::StaticClass());
+	AddOne->AllocateDefaultPins();
+
+	UEdGraphPin* Add_A = AddOne->FindPinChecked(TEXT("A"));
+	UEdGraphPin* Add_B = AddOne->FindPinChecked(TEXT("B"));
+	UEdGraphPin* Add_Return = AddOne->GetReturnValuePin();
+
+	Temp_Variable->MakeLinkTo(Add_A);
+	Add_B->DefaultValue = TEXT("1");
+	Add_Return->MakeLinkTo(Inc_Value);
+
+	///////////////////////////////////////////////////////////////////////////////////
+	// Create a sequence from the break exec that will set the loop counter to the last array index.
+	// The loop will then increment the counter and terminate on the next run of SequenceTwo.
+
+	UK2Node_AssignmentStatement* SetVariable = CompilerContext.SpawnIntermediateNode<UK2Node_AssignmentStatement>(this, SourceGraph);
+	SetVariable->AllocateDefaultPins();
+
+	UEdGraphPin* Set_Exec = SetVariable->GetExecPin();
+	UEdGraphPin* Set_Variable = SetVariable->GetVariablePin();
+	UEdGraphPin* Set_Value = SetVariable->GetValuePin();
+
+	CompilerContext.MovePinLinksToIntermediate(*ForEach_Break, *Set_Exec);
+	K2Schema->TryCreateConnection(Temp_Variable, Set_Variable);
+
+	UK2Node_CallFunction* GetSetLastIndex = CompilerContext.SpawnIntermediateNode<UK2Node_CallFunction>(this, SourceGraph);
+	GetSetLastIndex->FunctionReference.SetExternalMember(GET_FUNCTION_NAME_CHECKED(UBlueprintSetLibrary, Set_GetLastIndex), UBlueprintSetLibrary::StaticClass());
+	GetSetLastIndex->AllocateDefaultPins();
+
+	UEdGraphPin* GetIndex_Set = GetSetLastIndex->FindPinChecked(TEXT("TargetSet"));
+	UEdGraphPin* GetIndex_Return = GetSetLastIndex->GetReturnValuePin();
+
+	// Coerce the wildcard pin types
+	GetIndex_Set->PinType = ForEach_Set->PinType;
+	CompilerContext.CopyPinLinksToIntermediate(*ForEach_Set, *GetIndex_Set);
+
+	GetIndex_Return->MakeLinkTo(Set_Value);
+
+	BreakAllNodeLinks();
+}
+
+FText UK2Node_SetForEach::GetNodeTitle(ENodeTitleType::Type TitleType) const
+{
+	return LOCTEXT("NodeTitle", "For Each Loop (Set)");
+}
+
+FText UK2Node_SetForEach::GetTooltipText() const
+{
+	return LOCTEXT("NodeToolTip", "Loop over each element of a Set");
+}
+
+FSlateIcon UK2Node_SetForEach::GetIconAndTint(FLinearColor& OutColor) const
+{
+	return FSlateIcon("EditorStyle", "GraphEditor.Macro.ForEach_16x");
+}
+
+void UK2Node_SetForEach::NotifyPinConnectionListChanged(UEdGraphPin* Pin)
+{
+	Super::NotifyPinConnectionListChanged(Pin);
+
+	check(Pin);
+
+	if (Pin->PinName == SetPinName)
+	{
+		RefreshWildcardPins();
+	}
+}
+
+void UK2Node_SetForEach::PostEditChangeProperty(FPropertyChangedEvent& PropertyChangedEvent)
+{
+	Super::PostEditChangeProperty(PropertyChangedEvent);
+
+	bool bRefresh = false;
+
+	if (PropertyChangedEvent.GetPropertyName() == GET_MEMBER_NAME_CHECKED(UK2Node_SetForEach, ValueName))
+	{
+		GetValuePin()->PinFriendlyName = FText::FromString(ValueName);
+		bRefresh = true;
+	}
+
+	if (bRefresh)
+	{
+		// Poke the graph to update the visuals based on the above changes
+		GetGraph()->NotifyGraphChanged();
+		FBlueprintEditorUtils::MarkBlueprintAsModified(GetBlueprint());
+	}
+}
+
+bool UK2Node_SetForEach::CheckForErrors(const FKismetCompilerContext& CompilerContext)
+{
+	bool bError = false;
+
+	if (GetSetPin()->LinkedTo.Num() == 0)
+	{
+		CompilerContext.MessageLog.Error(*LOCTEXT("MissingSet_Error", "For Each (Set) node @@ must have a Set to iterate." ).ToString(), this);
+		bError = true;
+	}
+
+	return bError;
+}
+
+void UK2Node_SetForEach::RefreshWildcardPins()
+{
+	UEdGraphPin* SetPin = GetSetPin();
+	check(SetPin);
+
+	UEdGraphPin* ValuePin = GetValuePin();
+	check(ValuePin);
+
+	const bool bIsWildcardPin = FWildcardNodeUtils::IsWildcardPin(SetPin);
+	if (bIsWildcardPin && SetPin->LinkedTo.Num() > 0)
+	{
+		if (const UEdGraphPin* InferrablePin = FWildcardNodeUtils::FindInferrableLinkedPin(SetPin))
+		{
+			FWildcardNodeUtils::InferType(SetPin, InferrablePin->PinType);
+
+			// In some contexts, the value pin may not be in a wildcard state (eg: paste operation).
+			// We'll just leave the pin with its current type and let the compiler catch any issues (if any).
+			// This also helps ensure that we don't jostle any pins that are in a split state.
+			if (FWildcardNodeUtils::IsWildcardPin(ValuePin))
+			{
+				FWildcardNodeUtils::InferType(ValuePin, InferrablePin->PinType);
+			}
+		}
+	}
+
+	// If no pins are connected, then we need to reset the dependent pins back to the original wildcard state.
+	if (SetPin->LinkedTo.Num() == 0)
+	{
+		FWildcardNodeUtils::ResetToWildcard(SetPin);
+		FWildcardNodeUtils::ResetToWildcard(ValuePin);
+	}
+}
+
+#undef LOCTEXT_NAMESPACE
diff --git a/Engine/Source/Runtime/Engine/Classes/Kismet/BlueprintSetLibrary.h b/Engine/Source/Runtime/Engine/Classes/Kismet/BlueprintSetLibrary.h
--- a/Engine/Source/Runtime/Engine/Classes/Kismet/BlueprintSetLibrary.h
+++ b/Engine/Source/Runtime/Engine/Classes/Kismet/BlueprintSetLibrary.h
@@ -151,6 +151,27 @@ public:
 	UFUNCTION(BlueprintCallable, CustomThunk, meta=(BlueprintInternalUseOnly = "true", SetParam = "Value"))
 	static ENGINE_API void SetSetPropertyByName(UObject* Object, FName PropertyName, const TSet<int32>& Value);
 
+	/**
+	 * Given a set and an index, returns a copy of the item found at that index.
+	 * Index may differ from the order in which they were added. This function is implemented
+	 * for the purpose of iterating sets by index.
+	 *
+	 * @param	TargetSet		The set to get an item from
+	 * @param	Index			The index in the set to get an item from
+	 * @return	A copy of the item stored at the index
+	 */
+	UFUNCTION(BlueprintPure, CustomThunk, meta = (DisplayName = "Get Item By Index", CompactNodeTitle = "GET ITEM BY INDEX", SetParam = "TargetSet|Item", AutoCreateRefTerm = "Item", BlueprintThreadSafe), Category = "Utilities|Set")
+	static void Set_GetItemByIndex(const TSet<int32>& TargetSet, int32 Index, int32& Item);
+
+	/** 
+	 * Retrieves the last valid index for a set
+	 *
+	 * @param	TargetSet		The set in question
+	 * @return	The number of entries in the set
+	 */
+	UFUNCTION(BlueprintPure, CustomThunk, meta=(DisplayName = "Last Index", CompactNodeTitle = "GET LAST INDEX", SetParam = "TargetSet", BlueprintThreadSafe), Category="Utilities|Set")
+	static int32 Set_GetLastIndex(const TSet<int32>& TargetSet);
+
 	DECLARE_FUNCTION(execSet_Add)
 	{
 		Stack.MostRecentProperty = nullptr;
@@ -545,6 +566,68 @@ public:
 		P_NATIVE_END;
 	}
 
+	DECLARE_FUNCTION(execSet_GetItemByIndex)
+	{
+		Stack.MostRecentProperty = nullptr;
+		Stack.StepCompiledIn<FSetProperty>(nullptr);
+		void* SetAddr = Stack.MostRecentPropertyAddress;
+		FSetProperty* SetProperty = CastField<FSetProperty>(Stack.MostRecentProperty);
+		if (!SetProperty)
+		{
+			Stack.bArrayContextFailed = true;
+			return;
+		}
+		P_GET_PROPERTY(FIntProperty, Index);
+
+		// Since Item isn't really an int, step the stack manually
+		const FProperty* ElementProp = SetProperty->ElementProp;
+		const int32 PropertySize = ElementProp->GetElementSize() * ElementProp->ArrayDim;
+		void* StorageSpace = FMemory_Alloca(PropertySize);
+		ElementProp->InitializeValue(StorageSpace);
+
+		Stack.MostRecentPropertyAddress = nullptr;
+		Stack.MostRecentPropertyContainer = nullptr;
+		Stack.StepCompiledIn<FProperty>(StorageSpace);
+		const FFieldClass* ElementPropClass = ElementProp->GetClass();
+		const FFieldClass* MostRecentPropClass = Stack.MostRecentProperty->GetClass();
+		void* ItemPtr;
+
+		// If the destination and the inner type are identical in size and their field classes derive from one another, then permit the writing out of the array element to the destination memory
+		if (Stack.MostRecentPropertyAddress != nullptr && (PropertySize == Stack.MostRecentProperty->GetElementSize() * Stack.MostRecentProperty->ArrayDim) &&
+			(MostRecentPropClass->IsChildOf(ElementPropClass) || ElementPropClass->IsChildOf(MostRecentPropClass)))
+		{
+			ItemPtr = Stack.MostRecentPropertyAddress;
+		}
+		else
+		{
+			ItemPtr = StorageSpace;
+		}
+
+		P_FINISH;
+		P_NATIVE_BEGIN;
+		GenericSet_GetItemByIndex(SetAddr, SetProperty, Index, ItemPtr);
+		P_NATIVE_END;
+		ElementProp->DestroyValue(StorageSpace);
+	}
+
+	DECLARE_FUNCTION(execSet_GetLastIndex)
+	{
+		Stack.MostRecentProperty = nullptr;
+		Stack.StepCompiledIn<FSetProperty>(nullptr);
+		void* SetAddr = Stack.MostRecentPropertyAddress;
+		FSetProperty* SetProperty = CastField<FSetProperty>(Stack.MostRecentProperty);
+		if (!SetProperty)
+		{
+			Stack.bArrayContextFailed = true;
+			return;
+		}
+	
+		P_FINISH;
+		P_NATIVE_BEGIN;
+		*(int32*)RESULT_PARAM = GenericSet_GetLastIndex(SetAddr, SetProperty);
+		P_NATIVE_END
+	}
+
 	static ENGINE_API void GenericSet_Add(const void* TargetSet, const FSetProperty* SetProperty, const void* ItemPtr);
 	static ENGINE_API void GenericSet_AddItems(const void* TargetSet, const FSetProperty* SetProperty, const void* TargetArray, const FArrayProperty* ArrayProperty);
 	static ENGINE_API bool GenericSet_Remove(const void* TargetSet, const FSetProperty* SetProperty, const void* ItemPtr);
@@ -559,8 +642,9 @@ public:
 	static ENGINE_API void GenericSet_Union(const void* SetA, const FSetProperty* SetPropertyA, const void* SetB, const FSetProperty* SetPropertyB, const void* SetResult, const FSetProperty* SetPropertyResult);
 	static ENGINE_API void GenericSet_Difference(const void* SetA, const FSetProperty* SetPropertyA, const void* SetB, const FSetProperty* SetPropertyB, const void* SetResult, const FSetProperty* SetPropertyResult);
 	static ENGINE_API void GenericSet_SetSetPropertyByName(UObject* OwnerObject, FName SetPropertyName, const void* SrcSetAddr);
+	static ENGINE_API void GenericSet_GetItemByIndex(void* TargetSet, const FSetProperty* SetProperty, int32 Index, void* Item);
+	static ENGINE_API int32 GenericSet_GetLastIndex(const void* TargetSet, const FSetProperty* SetProperty);
 
 private:
 	static constexpr int32 MaxSupportedSetSize = TNumericLimits<int32>::Max();
 };
-
diff --git a/Engine/Source/Runtime/Engine/Private/BlueprintSetLibrary.cpp b/Engine/Source/Runtime/Engine/Private/BlueprintSetLibrary.cpp
--- a/Engine/Source/Runtime/Engine/Private/BlueprintSetLibrary.cpp
+++ b/Engine/Source/Runtime/Engine/Private/BlueprintSetLibrary.cpp
@@ -201,3 +201,39 @@ void UBlueprintSetLibrary::GenericSet_SetSetPropertyByName(UObject* OwnerObject,
 		}
 	}
 }
+
+void UBlueprintSetLibrary::GenericSet_GetItemByIndex(void* TargetSet, const FSetProperty* SetProperty, int32 Index, void* Item)
+{
+	if (TargetSet)
+	{
+		FScriptSetHelper SetHelper(SetProperty, TargetSet);
+		const FProperty* ElementProp = SetProperty->ElementProp;
+		const TScriptContainerIterator<FScriptSetHelper> Iterator = SetHelper.CreateIterator(Index);
+
+		if (Iterator)
+		{
+			ElementProp->CopySingleValueToScriptVM(Item, SetHelper.GetElementPtr(Iterator));
+		}
+		else
+		{
+			FFrame::KismetExecutionMessage(*FString::Printf(TEXT("Attempted to access invalid index %d from set '%s' of length %d in '%s'!"),
+				Index,
+				*SetProperty->GetName(),
+				SetHelper.Num(),
+				*SetProperty->GetOwnerVariant().GetPathName()),
+				ELogVerbosity::Warning,
+				FName(TEXT("GetOutOfBoundsWarning")));
+		}
+	}
+}
+
+int32 UBlueprintSetLibrary::GenericSet_GetLastIndex(const void* TargetSet, const FSetProperty* SetProperty)
+{
+	if (TargetSet)
+	{
+		FScriptSetHelper SetHelper(SetProperty, TargetSet);
+		return SetHelper.Num() - 1;
+	}
+
+	return 0;
+}
