PR #10355: Create a blueprint node for iterating TMap's in a ForEach style similar to the ForEach for TArray.

Additional changes made by Epic:

* Used minimal API to reduce export count.
* Changed PinConnectionListChanged to NotifyPinConnectionListChanged.
* Replaced type state-tracking with RefreshWildcardPins, which updates the pins based on the current set link. This also simplifies the implementation.
* Removed pin tooltip customization.
* Added an InferType overload for setting a pin type from a FEdGraphTerminalType. Also fixed a bug in ResetToWildcard that didn't reset value pin data.

diff --git a/Engine/Source/Editor/BlueprintGraph/Classes/K2Node_MapForEach.h b/Engine/Source/Editor/BlueprintGraph/Classes/K2Node_MapForEach.h
new file mode 100644
--- /dev/null
+++ b/Engine/Source/Editor/BlueprintGraph/Classes/K2Node_MapForEach.h
@@ -0,0 +1,70 @@
+// Copyright Epic Games, Inc. All Rights Reserved.
+
+#pragma once
+
+#include "K2Node.h"
+
+#include "K2Node_MapForEach.generated.h"
+
+/** Custom blueprint node for iterating Maps in blueprints */
+UCLASS(MinimalAPI, CollapseCategories)
+class UK2Node_MapForEach : public UK2Node
+{
+	GENERATED_BODY()
+
+public:
+
+	UK2Node_MapForEach();
+
+	/** Pin Accessors */
+	[[nodiscard]] UEdGraphPin* GetMapPin() const;
+	[[nodiscard]] UEdGraphPin* GetBreakPin() const;
+	[[nodiscard]] UEdGraphPin* GetForEachPin() const;
+	[[nodiscard]] UEdGraphPin* GetKeyPin() const;
+	[[nodiscard]] UEdGraphPin* GetValuePin() const;
+	[[nodiscard]] UEdGraphPin* GetCompletedPin() const;
+
+	/** UK2Node interface */
+	[[nodiscard]] virtual bool IsNodeSafeToIgnore() const override { return true; }
+	virtual void GetMenuActions(FBlueprintActionDatabaseRegistrar& ActionRegistrar) const override;
+	[[nodiscard]] virtual FText GetMenuCategory() const override;
+	virtual void PostReconstructNode() override;
+	/** End of UK2Node interface */
+
+	/** UEdGraphNode interface */
+	virtual void AllocateDefaultPins() override;
+	virtual void ExpandNode(FKismetCompilerContext& CompilerContext, UEdGraph* SourceGraph) override;
+	[[nodiscard]] virtual FText GetNodeTitle(ENodeTitleType::Type TitleType) const override;
+	[[nodiscard]] virtual FText GetTooltipText() const override;
+	[[nodiscard]] virtual FSlateIcon GetIconAndTint(FLinearColor& OutColor) const override;
+	virtual void NotifyPinConnectionListChanged(UEdGraphPin* Pin) override;
+	[[nodiscard]] virtual bool ShouldShowNodeProperties() const override { return true; }
+	/** End of UEdGraphNode interface */
+
+	/** UObject interface */
+	void PostEditChangeProperty(FPropertyChangedEvent& PropertyChangedEvent) override;
+	/** End of UObject interface */
+
+private:
+
+	/** Pin Names */
+	static const FName MapPinName;
+	static const FName BreakPinName;
+	static const FName KeyPinName;
+	static const FName ValuePinName;
+	static const FName CompletedPinName;
+
+	/** Determine if there is any configuration options that shouldn't be allowed */
+	[[nodiscard]] bool CheckForErrors(const FKismetCompilerContext& CompilerContext);
+
+	/** Updates the wildcard pins based on current links */
+	void RefreshWildcardPins();
+
+	/** A user editable hook for the display name of the key pin */
+	UPROPERTY(EditDefaultsOnly, Category = "Map For Each")
+	FString KeyName;
+
+	/** A user editable hook for the display name of the value pin */
+	UPROPERTY(EditDefaultsOnly, Category = "Map For Each")
+	FString ValueName;
+};
diff --git a/Engine/Source/Editor/BlueprintGraph/Private/K2Node_MapForEach.cpp b/Engine/Source/Editor/BlueprintGraph/Private/K2Node_MapForEach.cpp
new file mode 100644
--- /dev/null
+++ b/Engine/Source/Editor/BlueprintGraph/Private/K2Node_MapForEach.cpp
@@ -0,0 +1,413 @@
+// Copyright Epic Games, Inc. All Rights Reserved.
+
+#include "K2Node_MapForEach.h"
+
+#include "BlueprintActionDatabaseRegistrar.h"
+#include "BlueprintNodeSpawner.h"
+#include "K2Node_AssignmentStatement.h"
+#include "K2Node_CallFunction.h"
+#include "K2Node_ExecutionSequence.h"
+#include "K2Node_IfThenElse.h"
+#include "K2Node_TemporaryVariable.h"
+#include "KismetCompiler.h"
+#include "Kismet/BlueprintMapLibrary.h"
+#include "Kismet/KismetMathLibrary.h"
+#include "Kismet2/BlueprintEditorUtils.h"
+#include "Kismet2/WildcardNodeUtils.h"
+
+#define LOCTEXT_NAMESPACE "K2Node_MapForEach"
+
+const FName UK2Node_MapForEach::MapPinName(TEXT("MapPin"));
+const FName UK2Node_MapForEach::BreakPinName(TEXT("BreakPin"));
+const FName UK2Node_MapForEach::KeyPinName(TEXT("KeyPin"));
+const FName UK2Node_MapForEach::ValuePinName(TEXT("ValuePin"));
+const FName UK2Node_MapForEach::CompletedPinName(TEXT("CompletedPin"));
+
+UK2Node_MapForEach::UK2Node_MapForEach()
+{
+	KeyName = LOCTEXT("KeyPin_FriendlyName", "Map Key").ToString();
+	ValueName = LOCTEXT("ValuePin_FriendlyName", "Map Value").ToString();
+}
+
+UEdGraphPin* UK2Node_MapForEach::GetMapPin() const
+{
+	return FindPinChecked(MapPinName);
+}
+
+UEdGraphPin* UK2Node_MapForEach::GetBreakPin() const
+{
+	return FindPinChecked(BreakPinName);
+}
+
+UEdGraphPin* UK2Node_MapForEach::GetForEachPin() const
+{
+	return FindPinChecked(UEdGraphSchema_K2::PN_Then);
+}
+
+UEdGraphPin* UK2Node_MapForEach::GetKeyPin() const
+{
+	return FindPinChecked(KeyPinName);
+}
+
+UEdGraphPin* UK2Node_MapForEach::GetValuePin() const
+{
+	return FindPinChecked(ValuePinName);
+}
+
+UEdGraphPin* UK2Node_MapForEach::GetCompletedPin() const
+{
+	return FindPinChecked(CompletedPinName);
+}
+
+void UK2Node_MapForEach::GetMenuActions(FBlueprintActionDatabaseRegistrar& ActionRegistrar) const
+{
+	// Actions are registered under specific object-keys; the idea is that 
+	// actions might have to be updated (or deleted) if their object-key is  
+	// mutated (or removed).
+	// Here, we use the node's class to keep from needlessly instantiating a UBlueprintNodeSpawner.
+	// Additionally, if the node type disappears, then the action should go with it.
+
+	UClass* ActionKey = GetClass();
+	if (ActionRegistrar.IsOpenForRegistration(ActionKey))
+	{
+		UBlueprintNodeSpawner* NodeSpawner = UBlueprintNodeSpawner::Create(ActionKey);
+		check(NodeSpawner != nullptr);
+		ActionRegistrar.AddBlueprintAction(ActionKey, NodeSpawner);
+	}
+}
+
+FText UK2Node_MapForEach::GetMenuCategory() const
+{
+	return LOCTEXT("NodeMenu", "Utilities|Map");
+}
+
+void UK2Node_MapForEach::PostReconstructNode()
+{
+	Super::PostReconstructNode();
+
+	RefreshWildcardPins();
+}
+
+void UK2Node_MapForEach::AllocateDefaultPins()
+{
+	Super::AllocateDefaultPins();
+
+	CreatePin(EGPD_Input, UEdGraphSchema_K2::PC_Exec, UEdGraphSchema_K2::PN_Execute);
+
+	UEdGraphNode::FCreatePinParams PinParams;
+	PinParams.ContainerType = EPinContainerType::Map;
+	PinParams.ValueTerminalType.TerminalCategory = UEdGraphSchema_K2::PC_Wildcard;
+
+	UEdGraphPin* MapPin = CreatePin(EGPD_Input, UEdGraphSchema_K2::PC_Wildcard, MapPinName, PinParams);
+	MapPin->PinType.bIsConst = true;
+	MapPin->PinType.bIsReference = true;
+	MapPin->PinFriendlyName = LOCTEXT("MapPin_FriendlyName", "Map");
+
+	UEdGraphPin* BreakPin = CreatePin(EGPD_Input, UEdGraphSchema_K2::PC_Exec, BreakPinName);
+	BreakPin->PinFriendlyName = LOCTEXT("BreakPin_FriendlyName", "Break");
+	BreakPin->bAdvancedView = true;
+
+	UEdGraphPin* ForEachPin = CreatePin(EGPD_Output, UEdGraphSchema_K2::PC_Exec, UEdGraphSchema_K2::PN_Then);
+	ForEachPin->PinFriendlyName = LOCTEXT("ForEachPin_FriendlyName", "Loop Body");
+
+	UEdGraphPin* KeyPin = CreatePin(EGPD_Output, UEdGraphSchema_K2::PC_Wildcard, KeyPinName);
+	KeyPin->PinFriendlyName = FText::FromString(KeyName);
+
+	UEdGraphPin* ValuePin = CreatePin(EGPD_Output, UEdGraphSchema_K2::PC_Wildcard, ValuePinName);
+	ValuePin->PinFriendlyName = FText::FromString(ValueName);
+
+	UEdGraphPin* CompletedPin = CreatePin(EGPD_Output, UEdGraphSchema_K2::PC_Exec, CompletedPinName);
+	CompletedPin->PinFriendlyName = LOCTEXT("CompletedPin_FriendlyName", "Completed");
+	CompletedPin->PinToolTip = LOCTEXT("CompletedPin_Tooltip", "Execution once all map elements have been visited").ToString();
+
+	if (AdvancedPinDisplay == ENodeAdvancedPins::NoPins)
+	{
+		AdvancedPinDisplay = ENodeAdvancedPins::Hidden;
+	}
+}
+
+void UK2Node_MapForEach::ExpandNode(FKismetCompilerContext& CompilerContext, UEdGraph* SourceGraph)
+{
+	Super::ExpandNode(CompilerContext, SourceGraph);
+
+	if (CheckForErrors(CompilerContext))
+	{
+		// Remove all the links to this node as they are no longer needed
+		BreakAllNodeLinks();
+		return;
+	}
+
+	const UEdGraphSchema_K2* K2Schema = CastChecked<UEdGraphSchema_K2>(GetSchema());
+
+	///////////////////////////////////////////////////////////////////////////////////
+	// Cache off versions of all our important pins
+
+	UEdGraphPin* ForEach_Exec = GetExecPin();
+	UEdGraphPin* ForEach_Map = GetMapPin();
+	UEdGraphPin* ForEach_Break = GetBreakPin();
+	UEdGraphPin* ForEach_ForEach = GetForEachPin();
+	UEdGraphPin* ForEach_Key = GetKeyPin();
+	UEdGraphPin* ForEach_Value = GetValuePin();
+	UEdGraphPin* ForEach_Completed = GetCompletedPin();
+
+	///////////////////////////////////////////////////////////////////////////////////
+	// Create a loop counter variable
+
+	UK2Node_TemporaryVariable* CreateTemporaryVariable = CompilerContext.SpawnIntermediateNode<UK2Node_TemporaryVariable>(this, SourceGraph);
+	CreateTemporaryVariable->VariableType.PinCategory = UEdGraphSchema_K2::PC_Int;
+	CreateTemporaryVariable->AllocateDefaultPins();
+
+	UEdGraphPin* Temp_Variable = CreateTemporaryVariable->GetVariablePin();
+
+	///////////////////////////////////////////////////////////////////////////////////
+	// Initialize the temporary to 0
+
+	UK2Node_AssignmentStatement* InitTemporaryVariable = CompilerContext.SpawnIntermediateNode<UK2Node_AssignmentStatement>(this, SourceGraph);
+	InitTemporaryVariable->AllocateDefaultPins();
+
+	UEdGraphPin* Init_Exec = InitTemporaryVariable->GetExecPin();
+	UEdGraphPin* Init_Variable = InitTemporaryVariable->GetVariablePin();
+	UEdGraphPin* Init_Value = InitTemporaryVariable->GetValuePin();
+	UEdGraphPin* Init_Then = InitTemporaryVariable->GetThenPin();
+
+	CompilerContext.MovePinLinksToIntermediate(*ForEach_Exec, *Init_Exec);
+	K2Schema->TryCreateConnection(Init_Variable, Temp_Variable);
+	Init_Value->DefaultValue = TEXT("0");
+
+	///////////////////////////////////////////////////////////////////////////////////
+	// Branch on comparing the loop index with the size of the map
+
+	UK2Node_IfThenElse* BranchOnIndex = CompilerContext.SpawnIntermediateNode<UK2Node_IfThenElse>(this, SourceGraph);
+	BranchOnIndex->AllocateDefaultPins();
+
+	UEdGraphPin* Branch_Exec = BranchOnIndex->GetExecPin();
+	UEdGraphPin* Branch_Input = BranchOnIndex->GetConditionPin();
+	UEdGraphPin* Branch_Then = BranchOnIndex->GetThenPin();
+	UEdGraphPin* Branch_Else = BranchOnIndex->GetElsePin();
+
+	Init_Then->MakeLinkTo(Branch_Exec);
+	CompilerContext.MovePinLinksToIntermediate(*ForEach_Completed, *Branch_Else);
+
+	UK2Node_CallFunction* CompareLessThan = CompilerContext.SpawnIntermediateNode<UK2Node_CallFunction>(this, SourceGraph);
+	CompareLessThan->FunctionReference.SetExternalMember(GET_FUNCTION_NAME_CHECKED(UKismetMathLibrary, Less_IntInt), UKismetMathLibrary::StaticClass());
+	CompareLessThan->AllocateDefaultPins();
+
+	UEdGraphPin* Compare_A = CompareLessThan->FindPinChecked(TEXT("A"));
+	UEdGraphPin* Compare_B = CompareLessThan->FindPinChecked(TEXT("B"));
+	UEdGraphPin* Compare_Return = CompareLessThan->GetReturnValuePin();
+
+	Branch_Input->MakeLinkTo(Compare_Return);
+	Temp_Variable->MakeLinkTo(Compare_A);
+
+	UK2Node_CallFunction* GetMapLength = CompilerContext.SpawnIntermediateNode<UK2Node_CallFunction>(this, SourceGraph);
+	GetMapLength->FunctionReference.SetExternalMember(GET_FUNCTION_NAME_CHECKED(UBlueprintMapLibrary, Map_Length), UBlueprintMapLibrary::StaticClass());
+	GetMapLength->AllocateDefaultPins();
+
+	UEdGraphPin* MapLength_Map = GetMapLength->FindPinChecked(TEXT("TargetMap"));
+	UEdGraphPin* MapLength_Return = GetMapLength->GetReturnValuePin();
+
+	// Coerce the wildcard pin types
+	MapLength_Map->PinType = ForEach_Map->PinType;
+
+	Compare_B->MakeLinkTo(MapLength_Return);
+	CompilerContext.CopyPinLinksToIntermediate(*ForEach_Map, *MapLength_Map);
+
+	///////////////////////////////////////////////////////////////////////////////////
+	// Sequence the loop body and incrementing the loop counter
+
+	UK2Node_ExecutionSequence* LoopSequence = CompilerContext.SpawnIntermediateNode<UK2Node_ExecutionSequence>(this, SourceGraph);
+	LoopSequence->AllocateDefaultPins();
+
+	UEdGraphPin* Sequence_Exec = LoopSequence->GetExecPin();
+	UEdGraphPin* Sequence_One = LoopSequence->GetThenPinGivenIndex(0);
+	UEdGraphPin* Sequence_Two = LoopSequence->GetThenPinGivenIndex(1);
+
+	Branch_Then->MakeLinkTo(Sequence_Exec);
+	CompilerContext.MovePinLinksToIntermediate(*ForEach_ForEach, *Sequence_One);
+
+	UK2Node_CallFunction* GetMapPair = CompilerContext.SpawnIntermediateNode<UK2Node_CallFunction>(this, SourceGraph);
+	GetMapPair->FunctionReference.SetExternalMember(GET_FUNCTION_NAME_CHECKED(UBlueprintMapLibrary, Map_GetKeyValueByIndex), UBlueprintMapLibrary::StaticClass());
+	GetMapPair->AllocateDefaultPins();
+
+	UEdGraphPin* GetPair_Map = GetMapPair->FindPinChecked(TEXT("TargetMap"));
+	UEdGraphPin* GetPair_Index = GetMapPair->FindPinChecked(TEXT("Index"));
+	UEdGraphPin* GetPair_Key = GetMapPair->FindPinChecked(TEXT("Key"));
+	UEdGraphPin* GetPair_Value = GetMapPair->FindPinChecked(TEXT("Value"));
+
+	// Coerce the wildcard pin types
+	GetPair_Map->PinType = ForEach_Map->PinType;
+	GetPair_Key->PinType = ForEach_Key->PinType;
+	GetPair_Value->PinType = ForEach_Value->PinType;
+
+	CompilerContext.CopyPinLinksToIntermediate(*ForEach_Map, *GetPair_Map);
+	GetPair_Index->MakeLinkTo(Temp_Variable);
+	CompilerContext.MovePinLinksToIntermediate(*ForEach_Key, *GetPair_Key);
+	CompilerContext.MovePinLinksToIntermediate(*ForEach_Value, *GetPair_Value);
+
+	///////////////////////////////////////////////////////////////////////////////////
+	// Increment the loop counter by one
+
+	UK2Node_AssignmentStatement* IncrementVariable = CompilerContext.SpawnIntermediateNode<UK2Node_AssignmentStatement>(this, SourceGraph);
+	IncrementVariable->AllocateDefaultPins();
+
+	UEdGraphPin* Inc_Exec = IncrementVariable->GetExecPin();
+	UEdGraphPin* Inc_Variable = IncrementVariable->GetVariablePin();
+	UEdGraphPin* Inc_Value = IncrementVariable->GetValuePin();
+	UEdGraphPin* Inc_Then = IncrementVariable->GetThenPin();
+
+	Sequence_Two->MakeLinkTo(Inc_Exec);
+	Branch_Exec->MakeLinkTo(Inc_Then);
+
+	K2Schema->TryCreateConnection(Temp_Variable, Inc_Variable);
+
+	UK2Node_CallFunction* AddOne = CompilerContext.SpawnIntermediateNode<UK2Node_CallFunction>(this, SourceGraph);
+	AddOne->FunctionReference.SetExternalMember(GET_FUNCTION_NAME_CHECKED(UKismetMathLibrary, Add_IntInt), UKismetMathLibrary::StaticClass());
+	AddOne->AllocateDefaultPins();
+
+	UEdGraphPin* Add_A = AddOne->FindPinChecked(TEXT("A"));
+	UEdGraphPin* Add_B = AddOne->FindPinChecked(TEXT("B"));
+	UEdGraphPin* Add_Return = AddOne->GetReturnValuePin();
+
+	Temp_Variable->MakeLinkTo(Add_A);
+	Add_B->DefaultValue = TEXT("1");
+	Add_Return->MakeLinkTo(Inc_Value);
+
+	///////////////////////////////////////////////////////////////////////////////////
+	// Create a sequence from the break exec that will set the loop counter to the last array index.
+	// The loop will then increment the counter and terminate on the next run of SequenceTwo.
+
+	UK2Node_AssignmentStatement* SetVariable = CompilerContext.SpawnIntermediateNode<UK2Node_AssignmentStatement>(this, SourceGraph);
+	SetVariable->AllocateDefaultPins();
+
+	UEdGraphPin* Set_Exec = SetVariable->GetExecPin();
+	UEdGraphPin* Set_Variable = SetVariable->GetVariablePin();
+	UEdGraphPin* Set_Value = SetVariable->GetValuePin();
+
+	CompilerContext.MovePinLinksToIntermediate(*ForEach_Break, *Set_Exec);
+	K2Schema->TryCreateConnection(Temp_Variable, Set_Variable);
+
+	UK2Node_CallFunction* GetMapLastIndex = CompilerContext.SpawnIntermediateNode<UK2Node_CallFunction>(this, SourceGraph);
+	GetMapLastIndex->FunctionReference.SetExternalMember(GET_FUNCTION_NAME_CHECKED(UBlueprintMapLibrary, Map_GetLastIndex), UBlueprintMapLibrary::StaticClass());
+	GetMapLastIndex->AllocateDefaultPins();
+
+	UEdGraphPin* GetIndex_Map = GetMapLastIndex->FindPinChecked(TEXT("TargetMap"));
+	UEdGraphPin* GetIndex_Return = GetMapLastIndex->GetReturnValuePin();
+
+	// Coerce the wildcard pin types
+	GetIndex_Map->PinType = ForEach_Map->PinType;
+	CompilerContext.CopyPinLinksToIntermediate(*ForEach_Map, *GetIndex_Map);
+
+	GetIndex_Return->MakeLinkTo(Set_Value);
+
+	BreakAllNodeLinks();
+}
+
+FText UK2Node_MapForEach::GetNodeTitle(ENodeTitleType::Type TitleType) const
+{
+	return LOCTEXT("NodeTitle", "For Each Loop (Map)");
+}
+
+FText UK2Node_MapForEach::GetTooltipText() const
+{
+	return LOCTEXT("NodeToolTip", "Loop over each element of a map");
+}
+
+FSlateIcon UK2Node_MapForEach::GetIconAndTint(FLinearColor& OutColor) const
+{
+	return FSlateIcon("EditorStyle", "GraphEditor.Macro.ForEach_16x");
+}
+
+void UK2Node_MapForEach::NotifyPinConnectionListChanged(UEdGraphPin* Pin)
+{
+	Super::NotifyPinConnectionListChanged(Pin);
+
+	check(Pin);
+
+	if (Pin->PinName == MapPinName)
+	{
+		RefreshWildcardPins();
+	}
+}
+
+void UK2Node_MapForEach::PostEditChangeProperty(FPropertyChangedEvent &PropertyChangedEvent)
+{
+	Super::PostEditChangeProperty(PropertyChangedEvent);
+
+	bool bRefresh = false;
+
+	if (PropertyChangedEvent.GetPropertyName() == GET_MEMBER_NAME_CHECKED(UK2Node_MapForEach, KeyName))
+	{
+		GetKeyPin()->PinFriendlyName = FText::FromString(KeyName);
+		bRefresh = true;
+
+	}
+	else if (PropertyChangedEvent.GetPropertyName() == GET_MEMBER_NAME_CHECKED(UK2Node_MapForEach, ValueName))
+	{
+		GetValuePin()->PinFriendlyName = FText::FromString(ValueName);
+		bRefresh = true;
+	}
+
+	if (bRefresh)
+	{
+		// Poke the graph to update the visuals based on the above changes
+		GetGraph()->NotifyGraphChanged();
+		FBlueprintEditorUtils::MarkBlueprintAsModified(GetBlueprint());
+	}
+}
+
+bool UK2Node_MapForEach::CheckForErrors(const FKismetCompilerContext& CompilerContext)
+{
+	bool bError = false;
+
+	if (GetMapPin()->LinkedTo.IsEmpty())
+	{
+		CompilerContext.MessageLog.Error(*LOCTEXT("MissingMap_Error", "For Each (Map) node @@ must have a Map to iterate.").ToString(), this);
+		bError = true;
+	}
+
+	return bError;
+}
+
+void UK2Node_MapForEach::RefreshWildcardPins()
+{
+	UEdGraphPin* MapPin = GetMapPin();
+	check(MapPin);
+
+	UEdGraphPin* KeyPin = GetKeyPin();
+	check(KeyPin);
+
+	UEdGraphPin* ValuePin = GetValuePin();
+	check(ValuePin);
+
+	const bool bIsWildcardPin = FWildcardNodeUtils::IsWildcardPin(MapPin);
+	if (bIsWildcardPin && MapPin->LinkedTo.Num() > 0)
+	{
+		if (const UEdGraphPin* InferrablePin = FWildcardNodeUtils::FindInferrableLinkedPin(MapPin))
+		{
+			FWildcardNodeUtils::InferType(MapPin, InferrablePin->PinType);
+
+			// In some contexts, the key/value pins may not be in a wildcard state (eg: paste operation).
+			// We'll just leave the pins with their current type and let the compiler catch any issues (if any).
+			// This also helps ensure that we don't jostle any pins that are in a split state.
+			if (FWildcardNodeUtils::IsWildcardPin(KeyPin))
+			{
+				FWildcardNodeUtils::InferType(KeyPin, InferrablePin->PinType);
+			}
+			
+			if (FWildcardNodeUtils::IsWildcardPin(ValuePin))
+			{
+				FWildcardNodeUtils::InferType(ValuePin->PinType, InferrablePin->PinType.PinValueType);
+			}
+		}
+	}
+
+	// If no pins are connected, then we need to reset the dependent pins back to the original wildcard state.
+	if (MapPin->LinkedTo.Num() == 0)
+	{
+		FWildcardNodeUtils::ResetToWildcard(MapPin);
+		FWildcardNodeUtils::ResetToWildcard(KeyPin);
+		FWildcardNodeUtils::ResetToWildcard(ValuePin);
+	}
+}
+
+#undef LOCTEXT_NAMESPACE
diff --git a/Engine/Source/Editor/UnrealEd/Private/Kismet2/WildcardNodeUtils.cpp b/Engine/Source/Editor/UnrealEd/Private/Kismet2/WildcardNodeUtils.cpp
--- a/Engine/Source/Editor/UnrealEd/Private/Kismet2/WildcardNodeUtils.cpp
+++ b/Engine/Source/Editor/UnrealEd/Private/Kismet2/WildcardNodeUtils.cpp
@@ -37,7 +37,7 @@ bool FWildcardNodeUtils::HasAnyWildcards(const UEdGraphPin* const Pin)
 {
 	check(Pin);
 	return Pin->PinType.PinCategory == UEdGraphSchema_K2::PC_Wildcard ||
-		Pin->PinType.PinValueType.TerminalCategory == UEdGraphSchema_K2::PC_Wildcard ;
+		Pin->PinType.PinValueType.TerminalCategory == UEdGraphSchema_K2::PC_Wildcard;
 }
 
 bool FWildcardNodeUtils::HasAnyNonWildcards(const UEdGraphPin* Pin)
@@ -121,7 +121,7 @@ void FWildcardNodeUtils::InferType(UEdGraphPin* ToPin, const FEdGraphPinType& Ty
 void FWildcardNodeUtils::InferType(FEdGraphPinType& ToType, const FEdGraphPinType& Type)
 {
 	bool bInferredSomething = false;
-	if(ToType.PinCategory == UEdGraphSchema_K2::PC_Wildcard)
+	if (ToType.PinCategory == UEdGraphSchema_K2::PC_Wildcard)
 	{
 		bInferredSomething = true;
 		ToType.PinCategory = Type.PinCategory;
@@ -129,7 +129,7 @@ void FWildcardNodeUtils::InferType(FEdGraphPinType& ToType, const FEdGraphPinTyp
 		ToType.PinSubCategoryObject = Type.PinSubCategoryObject;
 		ToType.bIsUObjectWrapper = Type.bIsUObjectWrapper;
 	}
-	if(ToType.PinValueType.TerminalCategory == UEdGraphSchema_K2::PC_Wildcard)
+	if (ToType.PinValueType.TerminalCategory == UEdGraphSchema_K2::PC_Wildcard)
 	{
 		bInferredSomething = true;
 		ToType.PinValueType = Type.PinValueType;
@@ -137,6 +137,20 @@ void FWildcardNodeUtils::InferType(FEdGraphPinType& ToType, const FEdGraphPinTyp
 	check(bInferredSomething);
 }
 
+void FWildcardNodeUtils::InferType(FEdGraphPinType& ToType, const FEdGraphTerminalType& Type)
+{
+	bool bInferredSomething = false;
+	if (ToType.PinCategory == UEdGraphSchema_K2::PC_Wildcard)
+	{
+		bInferredSomething = true;
+		ToType.PinCategory = Type.TerminalCategory;
+		ToType.PinSubCategory = Type.TerminalSubCategory;
+		ToType.PinSubCategoryObject = Type.TerminalSubCategoryObject;
+		ToType.bIsUObjectWrapper = Type.bTerminalIsUObjectWrapper;
+	}
+	check(bInferredSomething);
+}
+
 void FWildcardNodeUtils::ResetToWildcard(UEdGraphPin* Pin)
 {
 	const UEdGraphSchema_K2* Schema = GetDefault<UEdGraphSchema_K2>();
@@ -156,5 +170,9 @@ void FWildcardNodeUtils::ResetToWildcard(FEdGraphPinType& PinType)
 	PinType.PinSubCategory = NAME_None;
 	PinType.PinSubCategoryObject = nullptr;
 	PinType.bIsUObjectWrapper = false;
+	PinType.PinValueType.TerminalCategory = UEdGraphSchema_K2::PC_Wildcard;
+	PinType.PinValueType.TerminalSubCategory = NAME_None;
+	PinType.PinValueType.TerminalSubCategoryObject = nullptr;
+	PinType.PinValueType.bTerminalIsUObjectWrapper = false;
 }
 
diff --git a/Engine/Source/Editor/UnrealEd/Public/Kismet2/WildcardNodeUtils.h b/Engine/Source/Editor/UnrealEd/Public/Kismet2/WildcardNodeUtils.h
--- a/Engine/Source/Editor/UnrealEd/Public/Kismet2/WildcardNodeUtils.h
+++ b/Engine/Source/Editor/UnrealEd/Public/Kismet2/WildcardNodeUtils.h
@@ -80,6 +80,7 @@ public:
 	 */
 	static UNREALED_API void InferType(UEdGraphPin* ToPin, const FEdGraphPinType& Type);
 	static UNREALED_API void InferType(FEdGraphPinType& ToType, const FEdGraphPinType& Type);
+	static UNREALED_API void InferType(FEdGraphPinType& ToType, const FEdGraphTerminalType& Type);
 
 	/** Utility functions for resetting a pin to wildcard, again, leaving container information unchanged */
 	static UNREALED_API void ResetToWildcard(UEdGraphPin* Pin);
diff --git a/Engine/Source/Runtime/Engine/Classes/Kismet/BlueprintMapLibrary.h b/Engine/Source/Runtime/Engine/Classes/Kismet/BlueprintMapLibrary.h
--- a/Engine/Source/Runtime/Engine/Classes/Kismet/BlueprintMapLibrary.h
+++ b/Engine/Source/Runtime/Engine/Classes/Kismet/BlueprintMapLibrary.h
@@ -112,11 +112,32 @@ public:
 	static ENGINE_API void Map_Clear(const TMap<int32, int32>& TargetMap);
 
 	/** 
-	* Not exposed to users. Supports setting a map property on an object by name.
-	*/
+	 * Not exposed to users. Supports setting a map property on an object by name.
+	 */
 	UFUNCTION(BlueprintCallable, CustomThunk, meta=(BlueprintInternalUseOnly = "true", MapParam = "Value"))
 	static ENGINE_API void SetMapPropertyByName(UObject* Object, FName PropertyName, const TMap<int32, int32>& Value);
 
+	/**
+	 * Given a map and an index, returns a copy of the key and value found at that index.
+	 * Index may differ from the order in which they were added. This function is implemented
+	 * for the purpose of iterating maps by index.
+	 *
+	 * @param	TargetMap	The map to get a key & value from
+	 * @param	Index		The index in the map to get the key & value
+	 * @param	Key			A copy of the key stored in the map at the index
+	 * @param	Value		A copy of the value stored in the map at the index
+	 */
+	UFUNCTION(BlueprintPure, CustomThunk, meta = (DisplayName = "Get Key Value By Index", CompactNodeTitle = "GET KEY VALUE BY INDEX", MapParam = "TargetMap", MapKeyParam = "Key", MapValueParam = "Value", AutoCreateRefTerm = "Key, Value", BlueprintThreadSafe), Category = "Utilities|Map")
+	static ENGINE_API void Map_GetKeyValueByIndex(const TMap<int32, int32>& TargetMap, int32 Index, int32& Key, int32& Value);
+
+	/** 
+	 * Determines the last valid index for a map
+	 *
+	 * @param	TargetMap		The map in question
+	 * @return	The number of entries in the map
+	 */
+	UFUNCTION(BlueprintPure, CustomThunk, meta=(DisplayName = "Get Last Index", CompactNodeTitle = "GET LAST INDEX", MapParam = "TargetMap", BlueprintThreadSafe), Category="Utilities|Map")
+	static ENGINE_API int32 Map_GetLastIndex(const TMap<int32, int32>& TargetMap);
 
 	DECLARE_FUNCTION(execMap_Add)
 	{
@@ -416,6 +437,86 @@ public:
 		P_NATIVE_END;
 	}
 
+	DECLARE_FUNCTION(execMap_GetKeyValueByIndex)
+	{
+		Stack.MostRecentProperty = nullptr;
+		Stack.StepCompiledIn<FMapProperty>(nullptr);
+		void* MapAddr = Stack.MostRecentPropertyAddress;
+		FMapProperty* MapProperty = CastField<FMapProperty>(Stack.MostRecentProperty);
+		if (!MapProperty)
+		{
+			Stack.bArrayContextFailed = true;
+			return;
+		}
+		P_GET_PROPERTY(FIntProperty, Index);
+		
+		// Since Key isn't really an int, step the stack manually
+		const FProperty* CurrKeyProp = MapProperty->KeyProp;
+		const int32 KeyPropertySize = CurrKeyProp->GetElementSize() * CurrKeyProp->ArrayDim;
+		void* KeyStorageSpace = FMemory_Alloca(KeyPropertySize);
+		CurrKeyProp->InitializeValue(KeyStorageSpace);
+
+		Stack.MostRecentPropertyAddress = nullptr;
+		Stack.MostRecentPropertyContainer = nullptr;
+		Stack.StepCompiledIn<FProperty>(KeyStorageSpace);
+		const FFieldClass* CurrKeyPropClass = CurrKeyProp->GetClass();
+		const FFieldClass* MostRecentPropClass = Stack.MostRecentProperty->GetClass();
+
+		// If the destination and the inner type are identical in size and their field classes derive from one another, then permit the writing out of the array element to the destination memory
+		void* KeyPtr = KeyStorageSpace;
+		if (Stack.MostRecentPropertyAddress != nullptr && (KeyPropertySize == Stack.MostRecentProperty->GetElementSize() * Stack.MostRecentProperty->ArrayDim) &&
+			(MostRecentPropClass->IsChildOf(CurrKeyPropClass) || CurrKeyPropClass->IsChildOf(MostRecentPropClass)))
+		{
+			KeyPtr = Stack.MostRecentPropertyAddress;
+		}
+
+		// Since the Value isn't really an int, step the stack manually
+		const FProperty* CurrValueProp = MapProperty->ValueProp;
+		const int32 ValuePropertySize = CurrValueProp->GetElementSize() * CurrValueProp->ArrayDim;
+		void* ValueStorageSpace = FMemory_Alloca(ValuePropertySize);
+		CurrValueProp->InitializeValue(ValueStorageSpace);
+		
+		Stack.MostRecentPropertyAddress = nullptr;
+		Stack.MostRecentPropertyContainer = nullptr;
+		Stack.StepCompiledIn<FProperty>(ValueStorageSpace);
+		const FFieldClass* CurrValuePropClass = CurrValueProp->GetClass();
+		MostRecentPropClass = Stack.MostRecentProperty->GetClass();
+
+		// If the destination and the inner type are identical in size and their field classes derive from one another, then permit the writing out of the array element to the destination memory
+		void* ItemPtr = ValueStorageSpace;
+		if (Stack.MostRecentPropertyAddress != nullptr && (ValuePropertySize == Stack.MostRecentProperty->GetElementSize() * Stack.MostRecentProperty->ArrayDim) &&
+			(MostRecentPropClass->IsChildOf(CurrValuePropClass) || CurrValuePropClass->IsChildOf(MostRecentPropClass)))
+		{
+			ItemPtr = Stack.MostRecentPropertyAddress;
+		}
+			
+		P_FINISH;
+		P_NATIVE_BEGIN;
+		GenericMap_GetKeyValueByIndex(MapAddr, MapProperty, Index, KeyPtr, ItemPtr);
+		P_NATIVE_END;
+
+		CurrValueProp->DestroyValue(ValueStorageSpace);
+		CurrKeyProp->DestroyValue(KeyStorageSpace);
+	}
+	
+	DECLARE_FUNCTION(execMap_GetLastIndex)
+	{
+		Stack.MostRecentProperty = nullptr;
+		Stack.StepCompiledIn<FMapProperty>(nullptr);
+		void *MapAddr = Stack.MostRecentPropertyAddress;
+		FMapProperty* MapProperty = CastField<FMapProperty>(Stack.MostRecentProperty);
+		if (!MapProperty)
+		{
+			Stack.bArrayContextFailed = true;
+			return;
+		}
+	
+		P_FINISH;
+		P_NATIVE_BEGIN;
+		*(int32*)RESULT_PARAM = GenericMap_GetLastIndex(MapAddr, MapProperty);
+		P_NATIVE_END
+	}
+
 	static ENGINE_API void GenericMap_Add(const void* TargetMap, const FMapProperty* MapProperty, const void* KeyPtr, const void* ValuePtr);
 	static ENGINE_API bool GenericMap_Remove(const void* TargetMap, const FMapProperty* MapProperty, const void* KeyPtr);
 	static ENGINE_API bool GenericMap_Find(const void* TargetMap, const FMapProperty* MapProperty, const void* KeyPtr, void* ValuePtr);
@@ -426,8 +527,9 @@ public:
 	static ENGINE_API bool GenericMap_IsNotEmpty(const void* TargetMap, const FMapProperty* MapProperty);
 	static ENGINE_API void GenericMap_Clear(const void* TargetMap, const FMapProperty* MapProperty);
 	static ENGINE_API void GenericMap_SetMapPropertyByName(UObject* OwnerObject, FName MapPropertyName, const void* SrcMapAddr);
+	static ENGINE_API void GenericMap_GetKeyValueByIndex(const void* TargetMap, const FMapProperty* MapProperty, int32 Index, void* KeyPtr, void* ValuePtr);
+	static ENGINE_API int32 GenericMap_GetLastIndex(const void* TargetMap, const FMapProperty* MapProperty);
 
 private:
 	static constexpr int32 MaxSupportedMapSize = TNumericLimits<int32>::Max();
 };
-
diff --git a/Engine/Source/Runtime/Engine/Private/BlueprintMapLibrary.cpp b/Engine/Source/Runtime/Engine/Private/BlueprintMapLibrary.cpp
--- a/Engine/Source/Runtime/Engine/Private/BlueprintMapLibrary.cpp
+++ b/Engine/Source/Runtime/Engine/Private/BlueprintMapLibrary.cpp
@@ -136,4 +136,42 @@ void UBlueprintMapLibrary::GenericMap_SetMapPropertyByName(UObject* OwnerObject,
 			MapProp->CopyValuesInternal(Dest, SrcMapAddr, 1);
 		}
 	}
-}
\ No newline at end of file
+}
+
+void UBlueprintMapLibrary::GenericMap_GetKeyValueByIndex(const void* TargetMap, const FMapProperty* MapProperty, int32 Index, void* KeyPtr, void* ValuePtr)
+{
+	if (TargetMap)
+	{
+		FScriptMapHelper MapHelper(MapProperty, TargetMap);
+		const FProperty* KeyProp = MapProperty->KeyProp;
+		const FProperty* ValueProp = MapProperty->ValueProp;
+		const TScriptContainerIterator<FScriptMapHelper> Iterator = MapHelper.CreateIterator(Index);
+		
+		if (Iterator)
+		{
+			KeyProp->CopySingleValueToScriptVM(KeyPtr, MapHelper.GetKeyPtr(Iterator));
+			ValueProp->CopySingleValueToScriptVM(ValuePtr, MapHelper.GetValuePtr(Iterator));
+		}
+		else
+		{
+			FFrame::KismetExecutionMessage(*FString::Printf(TEXT("Attempted to access index %d from map '%s' of length %d in '%s'!"),
+				Index,
+				*MapProperty->GetName(),
+				MapHelper.Num(),
+				*MapProperty->GetOwnerVariant().GetPathName()),
+				ELogVerbosity::Warning,
+				FName(TEXT("GetOutOfBoundsWarning")));
+		}
+	}
+}
+
+int32 UBlueprintMapLibrary::GenericMap_GetLastIndex(const void* TargetMap, const FMapProperty* MapProperty)
+{
+	if (TargetMap)
+	{
+		FScriptMapHelper MapHelper(MapProperty, TargetMap);
+		return MapHelper.Num() - 1;
+	}
+
+	return 0;
+}
