Exported the primary pin functions of the set/map iterator nodes.

Follow-up request to CL 41843453.

In some cases (eg: custom node expansion), dependent modules might generate/use iterator nodes. Without the ability to directly access the node's primary pin functions, it can be a little clumsy to use.

#rnx

diff --git a/Engine/Source/Editor/BlueprintGraph/Classes/K2Node_MapForEach.h b/Engine/Source/Editor/BlueprintGraph/Classes/K2Node_MapForEach.h
--- a/Engine/Source/Editor/BlueprintGraph/Classes/K2Node_MapForEach.h
+++ b/Engine/Source/Editor/BlueprintGraph/Classes/K2Node_MapForEach.h
@@ -17,12 +17,12 @@ public:
 	UK2Node_MapForEach();
 
 	/** Pin Accessors */
-	[[nodiscard]] UEdGraphPin* GetMapPin() const;
-	[[nodiscard]] UEdGraphPin* GetBreakPin() const;
-	[[nodiscard]] UEdGraphPin* GetForEachPin() const;
-	[[nodiscard]] UEdGraphPin* GetKeyPin() const;
-	[[nodiscard]] UEdGraphPin* GetValuePin() const;
-	[[nodiscard]] UEdGraphPin* GetCompletedPin() const;
+	[[nodiscard]] BLUEPRINTGRAPH_API UEdGraphPin* GetMapPin() const;
+	[[nodiscard]] BLUEPRINTGRAPH_API UEdGraphPin* GetBreakPin() const;
+	[[nodiscard]] BLUEPRINTGRAPH_API UEdGraphPin* GetForEachPin() const;
+	[[nodiscard]] BLUEPRINTGRAPH_API UEdGraphPin* GetKeyPin() const;
+	[[nodiscard]] BLUEPRINTGRAPH_API UEdGraphPin* GetValuePin() const;
+	[[nodiscard]] BLUEPRINTGRAPH_API UEdGraphPin* GetCompletedPin() const;
 
 	/** UK2Node interface */
 	[[nodiscard]] virtual bool IsNodeSafeToIgnore() const override { return true; }
@@ -42,7 +42,7 @@ public:
 	/** End of UEdGraphNode interface */
 
 	/** UObject interface */
-	void PostEditChangeProperty(FPropertyChangedEvent& PropertyChangedEvent) override;
+	virtual void PostEditChangeProperty(FPropertyChangedEvent& PropertyChangedEvent) override;
 	/** End of UObject interface */
 
 private:
diff --git a/Engine/Source/Editor/BlueprintGraph/Classes/K2Node_SetForEach.h b/Engine/Source/Editor/BlueprintGraph/Classes/K2Node_SetForEach.h
--- a/Engine/Source/Editor/BlueprintGraph/Classes/K2Node_SetForEach.h
+++ b/Engine/Source/Editor/BlueprintGraph/Classes/K2Node_SetForEach.h
@@ -17,11 +17,11 @@ public:
 	UK2Node_SetForEach();
 
 	/** Pin Accessors */
-	[[nodiscard]] UEdGraphPin* GetSetPin() const;
-	[[nodiscard]] UEdGraphPin* GetBreakPin() const;
-	[[nodiscard]] UEdGraphPin* GetForEachPin() const;
-	[[nodiscard]] UEdGraphPin* GetValuePin() const;
-	[[nodiscard]] UEdGraphPin* GetCompletedPin() const;
+	[[nodiscard]] BLUEPRINTGRAPH_API UEdGraphPin* GetSetPin() const;
+	[[nodiscard]] BLUEPRINTGRAPH_API UEdGraphPin* GetBreakPin() const;
+	[[nodiscard]] BLUEPRINTGRAPH_API UEdGraphPin* GetForEachPin() const;
+	[[nodiscard]] BLUEPRINTGRAPH_API UEdGraphPin* GetValuePin() const;
+	[[nodiscard]] BLUEPRINTGRAPH_API UEdGraphPin* GetCompletedPin() const;
 
 	/** UK2Node interface */
 	[[nodiscard]] virtual bool IsNodeSafeToIgnore() const override { return true; }
@@ -41,7 +41,7 @@ public:
 	/** End of UEdGraphNode interface */
 
 	/** UObject interface */
-	void PostEditChangeProperty(FPropertyChangedEvent& PropertyChangedEvent) override;
+	virtual void PostEditChangeProperty(FPropertyChangedEvent& PropertyChangedEvent) override;
 	/** End of UObject interface */
 
 private:
