Function renaming can occasionally trip an error on "Create Event" nodes

"Create Events" nodes will validate that every connected delegate uses the same signature. If any don't match, then an error is reported on the node.

However, we can erroneously trip this if one of the links is to a dangling reroute node. UK2Node_CreateDelegate::IsValid will fail to find a signature and assume that there's an error.

In these cases, we can use FindFirstCompilerRelevantLinkedPin to see if there's actually a valid pin connection to check. If not, the assumption is that it's a dangling connection, which we can skip.

#rnx

diff --git a/Engine/Source/Editor/BlueprintGraph/Private/K2Node_CreateDelegate.cpp b/Engine/Source/Editor/BlueprintGraph/Private/K2Node_CreateDelegate.cpp
--- a/Engine/Source/Editor/BlueprintGraph/Private/K2Node_CreateDelegate.cpp
+++ b/Engine/Source/Editor/BlueprintGraph/Private/K2Node_CreateDelegate.cpp
@@ -80,7 +80,7 @@ bool UK2Node_CreateDelegate::IsValid(FString* OutMsg, bool bDontUseSkeletalClass
 	}
 
 	const UEdGraphPin* DelegatePin = GetDelegateOutPin();
-	if(!DelegatePin)
+	if (!DelegatePin)
 	{
 		if (OutMsg)
 		{
@@ -90,7 +90,7 @@ bool UK2Node_CreateDelegate::IsValid(FString* OutMsg, bool bDontUseSkeletalClass
 	}
 
 	const UFunction* Signature = GetDelegateSignature();
-	if(!Signature)
+	if (!Signature)
 	{
 		if (OutMsg)
 		{
@@ -99,33 +99,38 @@ bool UK2Node_CreateDelegate::IsValid(FString* OutMsg, bool bDontUseSkeletalClass
 		return false;
 	}
 
-	for(int PinIter = 1; PinIter < DelegatePin->LinkedTo.Num(); PinIter++)
+	for (int PinIter = 1; PinIter < DelegatePin->LinkedTo.Num(); PinIter++)
 	{
-		const UEdGraphPin* OtherPin = DelegatePin->LinkedTo[PinIter];
-		const UFunction* OtherSignature = OtherPin ? 
-			FMemberReference::ResolveSimpleMemberReference<UFunction>(OtherPin->PinType.PinSubCategoryMemberReference) : NULL;
-		if(!OtherSignature || !Signature->IsSignatureCompatibleWith(OtherSignature))
+		UEdGraphPin* OtherPin = FBlueprintEditorUtils::FindFirstCompilerRelevantLinkedPin(DelegatePin->LinkedTo[PinIter]);
+		
+		// A null linked pin can occur if we run across an unconnected knot.
+		// In this case, it's expected that it wouldn't have a valid function signature to check.
+		if (OtherPin)
 		{
-			if (OutMsg)
+			const UFunction* OtherSignature = FMemberReference::ResolveSimpleMemberReference<UFunction>(OtherPin->PinType.PinSubCategoryMemberReference);
+			if (!OtherSignature || !Signature->IsSignatureCompatibleWith(OtherSignature))
 			{
-				if (const UK2Node_BaseMCDelegate* DelegateNode = OtherPin ? Cast<const UK2Node_BaseMCDelegate>(OtherPin->GetOwningNode()) : nullptr)
+				if (OutMsg)
 				{
-					const FString DelegateName = DelegateNode->GetPropertyName().ToString();
+					if (const UK2Node_BaseMCDelegate* DelegateNode = OtherPin ? Cast<const UK2Node_BaseMCDelegate>(OtherPin->GetOwningNode()) : nullptr)
+					{
+						const FString DelegateName = DelegateNode->GetPropertyName().ToString();
 
-					*OutMsg = FText::Format(NSLOCTEXT("K2Node", "Bad_delegate_connection_named_fmt", "A connected delegate ({0}) has an incompatible signature - has that delegate changed?"),
-						FText::FromString(DelegateName)).ToString();
-				}
-				else
-				{
-					*OutMsg = NSLOCTEXT("K2Node", "Bad_delegate_connection", "A connected delegate's signature is incompatible - has that delegate changed?").ToString();
+						*OutMsg = FText::Format(NSLOCTEXT("K2Node", "Bad_delegate_connection_named_fmt", "A connected delegate ({0}) has an incompatible signature - has that delegate changed?"),
+							FText::FromString(DelegateName)).ToString();
+					}
+					else
+					{
+						*OutMsg = NSLOCTEXT("K2Node", "Bad_delegate_connection", "A connected delegate's signature is incompatible - has that delegate changed?").ToString();
+					}
 				}
+				return false;
 			}
-			return false;
 		}
 	}
 
 	UClass* ScopeClass = GetScopeClass(bDontUseSkeletalClassForSelf);
-	if(!ScopeClass)
+	if (!ScopeClass)
 	{
 		if (OutMsg)
 		{
@@ -176,15 +181,15 @@ bool UK2Node_CreateDelegate::IsValid(FString* OutMsg, bool bDontUseSkeletalClass
 		return false;
 	}
 
-	if(!FoundFunction->HasAllFunctionFlags(FUNC_BlueprintAuthorityOnly))
+	if (!FoundFunction->HasAllFunctionFlags(FUNC_BlueprintAuthorityOnly))
 	{
-		for(int PinIter = 0; PinIter < DelegatePin->LinkedTo.Num(); PinIter++)
+		for (int PinIter = 0; PinIter < DelegatePin->LinkedTo.Num(); PinIter++)
 		{
 			const UEdGraphPin* OtherPin = DelegatePin->LinkedTo[PinIter];
 			const UK2Node_BaseMCDelegate* Node = OtherPin ? Cast<const UK2Node_BaseMCDelegate>(OtherPin->GetOwningNode()) : NULL;
 			if (Node && Node->IsAuthorityOnly())
 			{
-				if(OutMsg)
+				if (OutMsg)
 				{
 					*OutMsg = FText::Format(NSLOCTEXT("K2Node", "WrongDelegateAuthorityOnlyFmt", "The selected function/event ('{0}') is not compatible with this delegate (the delegate is server-only) - try marking the function/event AuthorityOnly."),
 						FText::FromString(FunctionName.ToString())).ToString();
