Add a note for event nodes that override const functions - these are causing implicit object mutation and can permit explicit mutation which should have been forbidden

#rnx

diff --git a/Engine/Source/Editor/BlueprintGraph/Private/K2Node_Event.cpp b/Engine/Source/Editor/BlueprintGraph/Private/K2Node_Event.cpp
--- a/Engine/Source/Editor/BlueprintGraph/Private/K2Node_Event.cpp
+++ b/Engine/Source/Editor/BlueprintGraph/Private/K2Node_Event.cpp
@@ -380,6 +380,13 @@ void UK2Node_Event::ValidateNodeDuringCompilation(class FCompilerResultsLog& Mes
 			// If we are overriding a function, but we can;t find the function we are overriding, that is a compile error
 			MessageLog.Error(*FText::Format(NSLOCTEXT("KismetCompiler", "MissingEventSig_ErrorFmt", "Missing Event '{0}' for @@"), FText::FromString(EventReference.GetMemberName().ToString())).ToString(), this);
 		}
+		else if(Function->HasAnyFunctionFlags(FUNC_Const|FUNC_BlueprintPure))
+		{
+			// The implicit writes of object parameters to the event graph frame can violate clustering rules, so this is 
+			// noting some quite severe shortcomings in the blueprint compiler - but starting with a note while we address
+			// those:
+			MessageLog.Note(*FText::Format(NSLOCTEXT("KismetCompiler", "EventOverridingConst_Fmt", "@@: Event overriding '{0}', which is const, implicitly writes to this object. Please convert this to a function - ESPECIALLY if this event has Object parameters"), FText::FromString(EventReference.GetMemberName().ToString())).ToString(), this);
+		}
 	}
 	else if (UBlueprint* Blueprint = GetBlueprint())
 	{
