Report changes to the delegate node's signature function

[RN] The blueprint diff tool now reports changes to Create Event Node's signature function

diff --git a/Engine/Source/Editor/BlueprintGraph/Classes/K2Node_CreateDelegate.h b/Engine/Source/Editor/BlueprintGraph/Classes/K2Node_CreateDelegate.h
--- a/Engine/Source/Editor/BlueprintGraph/Classes/K2Node_CreateDelegate.h
+++ b/Engine/Source/Editor/BlueprintGraph/Classes/K2Node_CreateDelegate.h
@@ -43,6 +43,7 @@ public:
 	virtual void NodeConnectionListChanged() override;
 	virtual FText GetNodeTitle(ENodeTitleType::Type TitleType) const override;
 	virtual UObject* GetJumpTargetForDoubleClick() const override;
+	virtual void FindDiffs(UEdGraphNode* OtherNode, struct FDiffResults& Results) override;
 	virtual void AddSearchMetaDataInfo(TArray<struct FSearchTagDataPair>& OutTaggedMetaData) const override;
 	// End of UEdGraphNode interface
 
diff --git a/Engine/Source/Editor/BlueprintGraph/Private/K2Node_CreateDelegate.cpp b/Engine/Source/Editor/BlueprintGraph/Private/K2Node_CreateDelegate.cpp
--- a/Engine/Source/Editor/BlueprintGraph/Private/K2Node_CreateDelegate.cpp
+++ b/Engine/Source/Editor/BlueprintGraph/Private/K2Node_CreateDelegate.cpp
@@ -5,6 +5,7 @@
 #include "BlueprintNodeSpawner.h"
 #include "Containers/UnrealString.h"
 #include "DelegateNodeHandlers.h"
+#include "DiffResults.h"
 #include "EdGraph/EdGraph.h"
 #include "EdGraph/EdGraphPin.h"
 #include "EdGraphSchema_K2.h"
@@ -463,6 +464,27 @@ UObject* UK2Node_CreateDelegate::GetJumpTargetForDoubleClick() const
 	return GetDelegateSignature();
 }
 
+void UK2Node_CreateDelegate::FindDiffs(UEdGraphNode* OtherNode, struct FDiffResults& Results)
+{
+	Super::FindDiffs(OtherNode, Results);
+
+	UK2Node_CreateDelegate* OtherCreateDelegate = Cast<UK2Node_CreateDelegate>(OtherNode);
+	if (OtherCreateDelegate && OtherCreateDelegate->SelectedFunctionName != SelectedFunctionName)
+	{		
+		FDiffSingleResult Diff;
+		Diff.Diff = EDiffType::NODE_PROPERTY;
+		Diff.Node1 = this;
+		Diff.Node2 = OtherNode;
+		Diff.DisplayString = FText::Format(
+			NSLOCTEXT("K2Node", "DIF_SignatureFunction", "Signature function changed from {0} to {1}"), 
+			FText::FromName(SelectedFunctionName), 
+			FText::FromName(OtherCreateDelegate->SelectedFunctionName));
+		Diff.Category = EDiffType::MODIFICATION;
+
+		Results.Add(Diff);
+	}
+}
+
 void UK2Node_CreateDelegate::AddSearchMetaDataInfo(TArray<struct FSearchTagDataPair>& OutTaggedMetaData) const
 {
 	Super::AddSearchMetaDataInfo(OutTaggedMetaData);
