UE-278419 - Changing variable type on existing variables with expanded packed data inputs crashes editor

When splitting structure pins, the pin can't be in a container. This behavior is undefined, and it causes a pin count mismatch in RewireOldPinsToNewPins. Specifically, UK2Node::RestoreSplitPins automatically skips container types, but TransferPersistentDataFromOldPin would erroneously split any container pins.

The easiest fix is to simply skip any container pins in UEdGraphSchema_K2::SplitPin.

diff --git a/Engine/Source/Editor/BlueprintGraph/Private/EdGraphSchema_K2.cpp b/Engine/Source/Editor/BlueprintGraph/Private/EdGraphSchema_K2.cpp
--- a/Engine/Source/Editor/BlueprintGraph/Private/EdGraphSchema_K2.cpp
+++ b/Engine/Source/Editor/BlueprintGraph/Private/EdGraphSchema_K2.cpp
@@ -7339,7 +7339,12 @@ void UEdGraphSchema_K2::SplitPin(UEdGraphPin* Pin, const bool bNotify) const
 {
 	// Under some circumstances we can get here when PinSubCategoryObject is not set, so we just can't split the pin in that case
 	UScriptStruct* StructType = Cast<UScriptStruct>(Pin->PinType.PinSubCategoryObject.Get());
-	if (StructType == nullptr)
+	const bool bCanSplit =
+		(StructType != nullptr) &&
+		(!Pin->PinType.IsContainer())
+	;
+
+	if (!bCanSplit)
 	{
 		return;
 	}
