Updated headers using UnrealCodeFixup to make sure dllstorage is on methods/staticvars instead of types. Part 4/n

diff --git a/Engine/Source/Editor/BlueprintGraph/Classes/BlueprintFunctionNodeSpawner.h b/Engine/Source/Editor/BlueprintGraph/Classes/BlueprintFunctionNodeSpawner.h
--- a/Engine/Source/Editor/BlueprintGraph/Classes/BlueprintFunctionNodeSpawner.h
+++ b/Engine/Source/Editor/BlueprintGraph/Classes/BlueprintFunctionNodeSpawner.h
@@ -16,6 +16,8 @@
 
 #include "BlueprintFunctionNodeSpawner.generated.h"
 
+#define UE_API BLUEPRINTGRAPH_API
+
 class UEdGraph;
 class UEdGraphNode;
 class UFunction;
@@ -27,8 +29,8 @@ struct FBlueprintActionContext;
  * Takes care of spawning various UK2Node_CallFunction nodes. Acts as the 
  * "action" portion of certain FBlueprintActionMenuItems. 
  */
-UCLASS(Transient)
-class BLUEPRINTGRAPH_API UBlueprintFunctionNodeSpawner : public UBlueprintFieldNodeSpawner
+UCLASS(MinimalAPI, Transient)
+class UBlueprintFunctionNodeSpawner : public UBlueprintFieldNodeSpawner
 {
 	GENERATED_UCLASS_BODY()
 
@@ -42,7 +44,7 @@ public:
 	 * @param  Outer		Optional outer for the new spawner (if left null, the transient package will be used).
 	 * @return A newly allocated instance of this class.
 	 */
-	static UBlueprintFunctionNodeSpawner* Create(UFunction const* const Function, UObject* Outer = nullptr);
+	static UE_API UBlueprintFunctionNodeSpawner* Create(UFunction const* const Function, UObject* Outer = nullptr);
 
 	/**
 	 * Creates a new UBlueprintFunctionNodeSpawner for the specified function.
@@ -54,18 +56,18 @@ public:
 	 * @param  Outer		Optional outer for the new spawner (if left null, the transient package will be used).
 	 * @return A newly allocated instance of this class.
 	 */
-	static UBlueprintFunctionNodeSpawner* Create(TSubclassOf<UK2Node_CallFunction> NodeClass, UFunction const* const Function, UObject* Outer = nullptr);
+	static UE_API UBlueprintFunctionNodeSpawner* Create(TSubclassOf<UK2Node_CallFunction> NodeClass, UFunction const* const Function, UObject* Outer = nullptr);
 
 	// UBlueprintNodeSpawner interface
-	virtual void Prime() override;
-	virtual FBlueprintActionUiSpec GetUiSpec(FBlueprintActionContext const& Context, FBindingSet const& Bindings) const override;
-	virtual UEdGraphNode* Invoke(UEdGraph* ParentGraph, FBindingSet const& Bindings, FVector2D const Location) const override;
+	UE_API virtual void Prime() override;
+	UE_API virtual FBlueprintActionUiSpec GetUiSpec(FBlueprintActionContext const& Context, FBindingSet const& Bindings) const override;
+	UE_API virtual UEdGraphNode* Invoke(UEdGraph* ParentGraph, FBindingSet const& Bindings, FVector2D const Location) const override;
 	// End UBlueprintNodeSpawner interface
 
 	// IBlueprintNodeBinder interface
-	virtual bool IsBindingCompatible(FBindingObject BindingCandidate) const override;
-	virtual bool CanBindMultipleObjects() const override;
-	virtual bool BindToNode(UEdGraphNode* Node, FBindingObject Binding) const override;
+	UE_API virtual bool IsBindingCompatible(FBindingObject BindingCandidate) const override;
+	UE_API virtual bool CanBindMultipleObjects() const override;
+	UE_API virtual bool BindToNode(UEdGraphNode* Node, FBindingObject Binding) const override;
 	// End IBlueprintNodeBinder interface
 
 	/**
@@ -74,5 +76,7 @@ public:
 	 *
 	 * @return The function that this class was initialized with.
 	 */
-	UFunction const* GetFunction() const;
+	UE_API UFunction const* GetFunction() const;
 };
+
+#undef UE_API
diff --git a/Engine/Source/Editor/BlueprintGraph/Classes/BlueprintTypePromotion.h b/Engine/Source/Editor/BlueprintGraph/Classes/BlueprintTypePromotion.h
--- a/Engine/Source/Editor/BlueprintGraph/Classes/BlueprintTypePromotion.h
+++ b/Engine/Source/Editor/BlueprintGraph/Classes/BlueprintTypePromotion.h
@@ -19,6 +19,8 @@
 #include "UObject/NameTypes.h"
 #include "UObject/UObjectGlobals.h"
 
+#define UE_API BLUEPRINTGRAPH_API
+
 class FText;
 class UBlueprintFunctionNodeSpawner;
 class UFunction;
@@ -28,42 +30,42 @@ class UFunction;
 * Creates a map of "Operations" to any of their matching UFunctions
 * so that we can find the best possible match given several pin types.  
 */
-class BLUEPRINTGRAPH_API FTypePromotion : private FNoncopyable
+class FTypePromotion : private FNoncopyable
 {
 public:
 
 	/** Creates a new singleton instance of TypePromotion if there isn't one and returns a reference to it */
-	static FTypePromotion& Get();
+	static UE_API FTypePromotion& Get();
 
 	/** Deletes the singleton instance of type promotion if there is one */
-	static void Shutdown();
+	static UE_API void Shutdown();
 
 	/**
 	* Find the function that is the best match given the pins to consider. 
 	* Ex: Given "Add" operator and an array of two Vector pins, it will return "Add_VectorVector"
 	*/
-	static UFunction* FindBestMatchingFunc(FName Operation, const TArray<UEdGraphPin*>& PinsToConsider);
+	static UE_API UFunction* FindBestMatchingFunc(FName Operation, const TArray<UEdGraphPin*>& PinsToConsider);
 
 	/** Returns all functions for a specific operation. Will empty the given array and populate it with UFunction pointers */
-	static void GetAllFuncsForOp(FName Operation, TArray<UFunction*>& OutFuncs);
+	static UE_API void GetAllFuncsForOp(FName Operation, TArray<UFunction*>& OutFuncs);
 
 	/** Get a set of the supported operator names for type promo. Ex: "Add", "Subtract", "Multiply" */
-	static const TSet<FName>& GetAllOpNames();
+	static UE_API const TSet<FName>& GetAllOpNames();
 	
 	/** Set of comparison operator names (GreaterThan, LessThan, etc) */
-	static const TSet<FName>& GetComparisonOpNames();
+	static UE_API const TSet<FName>& GetComparisonOpNames();
 
 	/** Get the keywords metadata for the given operator name */
-	static const FText& GetKeywordsForOperator(const FName Operator);
+	static UE_API const FText& GetKeywordsForOperator(const FName Operator);
 
 	/** Get the user facing version of this operator name */
-	static const FText& GetUserFacingOperatorName(const FName Operator);
+	static UE_API const FText& GetUserFacingOperatorName(const FName Operator);
 
 	/** Returns true if the given function is a comparison operator */
-	static bool IsComparisonFunc(UFunction const* const Func);
+	static UE_API bool IsComparisonFunc(UFunction const* const Func);
 
 	/** Returns true if the given op name is a comparison operator name */
-	static bool IsComparisonOpName(const FName OpName);
+	static UE_API bool IsComparisonOpName(const FName OpName);
 
 	/**
 	* Parse the name of the operator that this function matches to (Add, Subtract, etc)
@@ -72,17 +74,17 @@ public:
 	*
 	* @return FName		Name of the operator for this function, NO_OP if does not exist.
 	*/
-	static FName GetOpNameFromFunction(UFunction const* const Func);
+	static UE_API FName GetOpNameFromFunction(UFunction const* const Func);
 
 	/**
 	* Determine what type a given set of wildcard pins would result in
 	*
 	* @return	Pin type that is the "highest" of all the given pins
 	*/
-	static FEdGraphPinType GetPromotedType(const TArray<UEdGraphPin*>& WildcardPins);
+	static UE_API FEdGraphPinType GetPromotedType(const TArray<UEdGraphPin*>& WildcardPins);
 
 	/** Returns true if the given function has been registered within the operator table */
-	static bool IsFunctionPromotionReady(const UFunction* const FuncToConsider);
+	static UE_API bool IsFunctionPromotionReady(const UFunction* const FuncToConsider);
 
 	/** Represents the possible results when comparing two types for promotion */
 	enum class ETypeComparisonResult : uint8
@@ -97,29 +99,29 @@ public:
 	* Given the two pin types check which pin type is higher. 
 	* Given two structs it will return equal, this does NOT compare PinDefaultSubobjects
 	*/
-	static ETypeComparisonResult GetHigherType(const FEdGraphPinType& A, const FEdGraphPinType& B);
+	static UE_API ETypeComparisonResult GetHigherType(const FEdGraphPinType& A, const FEdGraphPinType& B);
 
 	/** Returns true if A can be promoted to type B correctly, or if the types are equal */
-	static bool IsValidPromotion(const FEdGraphPinType& A, const FEdGraphPinType& B);
+	static UE_API bool IsValidPromotion(const FEdGraphPinType& A, const FEdGraphPinType& B);
 
 	/** Returns true if the given input pin can correctly be converted to the output type as a struct */
-	static bool HasStructConversion(const UEdGraphPin* InputPin, const UEdGraphPin* OutputPin);
+	static UE_API bool HasStructConversion(const UEdGraphPin* InputPin, const UEdGraphPin* OutputPin);
 
 	/** Returns true if the given function has a registered operator node spawner */
-	static bool IsOperatorSpawnerRegistered(UFunction const* const Func);
+	static UE_API bool IsOperatorSpawnerRegistered(UFunction const* const Func);
 
 	/** Function node spawner associated with this operation */
-	static UBlueprintFunctionNodeSpawner* GetOperatorSpawner(FName OpName);
+	static UE_API UBlueprintFunctionNodeSpawner* GetOperatorSpawner(FName OpName);
 
 	/** keep track of the operator that this function provides so that we don't add multiple 
 	to the BP context menu */
-	static void RegisterOperatorSpawner(FName OpName, UBlueprintFunctionNodeSpawner* Spawner);
+	static UE_API void RegisterOperatorSpawner(FName OpName, UBlueprintFunctionNodeSpawner* Spawner);
 
 	/**
 	* Clear all registered node spawners for operators from the system. This is necessary
 	* during hot reload to ensure that all operator spawners are up to date.
 	*/
-	static void ClearNodeSpawners();
+	static UE_API void ClearNodeSpawners();
 
 	/**
 	 * Clear out the promotion table and rebuild it, reassessing each available UFunction.
@@ -127,7 +129,7 @@ public:
 	 * @param Reason	The reason for refreshing the promotion tables, which can indicate
 	 *					whether or not it came from a Live Coding recompile.
 	 */
-	static void RefreshPromotionTables(EReloadCompleteReason Reason = EReloadCompleteReason::None);
+	static UE_API void RefreshPromotionTables(EReloadCompleteReason Reason = EReloadCompleteReason::None);
 
 	/**
 	* Get the "Primitive Promotion Table" which represents what base Pin types 
@@ -135,7 +137,7 @@ public:
 	*
 	* @return Const pointer to the primitive promotion table
 	*/
-	static const TMap<FName, TArray<FName>>* const GetPrimitivePromotionTable();
+	static UE_API const TMap<FName, TArray<FName>>* const GetPrimitivePromotionTable();
 	
 	/**
 	* Get a pointer to an array of available promotion types to a given pin type
@@ -144,40 +146,40 @@ public:
 	*
 	* @return const TArray<FName>*	Pointer to the available promotions, nullptr if there are none.
 	*/
-	static const TArray<FName>* GetAvailablePrimitivePromotions(const FEdGraphPinType& Type);
+	static UE_API const TArray<FName>* GetAvailablePrimitivePromotions(const FEdGraphPinType& Type);
 
 private:
 
-	explicit FTypePromotion();
-	~FTypePromotion();
+	UE_API explicit FTypePromotion();
+	UE_API ~FTypePromotion();
 
-	bool IsFunctionPromotionReady_Internal(const UFunction* const FuncToConsider) const;
+	UE_API bool IsFunctionPromotionReady_Internal(const UFunction* const FuncToConsider) const;
 
-	FEdGraphPinType GetPromotedType_Internal(const TArray<UEdGraphPin*>& WildcardPins) const;
+	UE_API FEdGraphPinType GetPromotedType_Internal(const TArray<UEdGraphPin*>& WildcardPins) const;
 	
-	UFunction* FindBestMatchingFunc_Internal(FName Operation, const TArray<UEdGraphPin*>& PinsToConsider);
+	UE_API UFunction* FindBestMatchingFunc_Internal(FName Operation, const TArray<UEdGraphPin*>& PinsToConsider);
 
-	void GetAllFuncsForOp_Internal(FName Operation, TArray<UFunction*>& OutFuncs);
+	UE_API void GetAllFuncsForOp_Internal(FName Operation, TArray<UFunction*>& OutFuncs);
 
 	/** Returns true if the given function is a candidate to handle type promotion */
-	static bool IsPromotableFunction(const UFunction* Function);
+	static UE_API bool IsPromotableFunction(const UFunction* Function);
 
 	/**
 	* Determines which pin type is "higher" 
 	* 
 	* @see FTypePromotion::PromotionTable
 	*/
-	ETypeComparisonResult GetHigherType_Internal(const FEdGraphPinType& A, const FEdGraphPinType& B) const;
+	UE_API ETypeComparisonResult GetHigherType_Internal(const FEdGraphPinType& A, const FEdGraphPinType& B) const;
 
 	/** Creates a lookup table of types and operations to their appropriate UFunction */
-	void CreateOpTable();
+	UE_API void CreateOpTable();
 
 	/** Creates the table of what types can be promoted to others */
-	TMap<FName, TArray<FName>> CreatePromotionTable();
+	UE_API TMap<FName, TArray<FName>> CreatePromotionTable();
 
-	void AddOpFunction(FName OpName, UFunction* Function);
+	UE_API void AddOpFunction(FName OpName, UFunction* Function);
 
-	static FTypePromotion* Instance;
+	static UE_API FTypePromotion* Instance;
 
 	/** A map of 'Type' to its 'available promotions'. See ctor for creation */
 	const TMap<FName, TArray<FName>> PromotionTable;
@@ -214,4 +216,6 @@ namespace TypePromoDebug
 	{
 		return GetDefault<UBlueprintEditorSettings>()->bEnableTypePromotion;
 	}
-}
\ No newline at end of file
+}
+
+#undef UE_API
diff --git a/Engine/Source/Editor/BlueprintGraph/Classes/BlueprintVariableNodeSpawner.h b/Engine/Source/Editor/BlueprintGraph/Classes/BlueprintVariableNodeSpawner.h
--- a/Engine/Source/Editor/BlueprintGraph/Classes/BlueprintVariableNodeSpawner.h
+++ b/Engine/Source/Editor/BlueprintGraph/Classes/BlueprintVariableNodeSpawner.h
@@ -22,6 +22,8 @@
 
 #include "BlueprintVariableNodeSpawner.generated.h"
 
+#define UE_API BLUEPRINTGRAPH_API
+
 class FProperty;
 class UClass;
 class UEdGraph;
@@ -36,8 +38,8 @@ struct FBlueprintActionContext;
  * FEdGraphSchemaAction_K2Var, and can spawn nodes for both member-variables and 
  * local function variables.
  */
-UCLASS(Transient)
-class BLUEPRINTGRAPH_API UBlueprintVariableNodeSpawner : public UBlueprintFieldNodeSpawner
+UCLASS(MinimalAPI, Transient)
+class UBlueprintVariableNodeSpawner : public UBlueprintFieldNodeSpawner
 {
 	GENERATED_UCLASS_BODY()
 
@@ -52,7 +54,7 @@ public:
 	 * @param  OwnerClass	The class that the variable is a member of or the class it is associated with if it is in a sidecar data structure.
 	 * @return A newly allocated instance of this class.
 	 */
-	static UBlueprintVariableNodeSpawner* CreateFromMemberOrParam(TSubclassOf<UK2Node_Variable> NodeClass, FProperty const* VarProperty, UEdGraph* VarContext = nullptr, UClass* OwnerClass = nullptr);
+	static UE_API UBlueprintVariableNodeSpawner* CreateFromMemberOrParam(TSubclassOf<UK2Node_Variable> NodeClass, FProperty const* VarProperty, UEdGraph* VarContext = nullptr, UClass* OwnerClass = nullptr);
 
 	/**
 	 * Creates a new UBlueprintVariableNodeSpawner, charged with spawning
@@ -65,19 +67,19 @@ public:
 	 * @param  Outer		Optional outer for the new spawner (if left null, the transient package will be used).
 	 * @return A newly allocated instance of this class.
 	 */
-	static UBlueprintVariableNodeSpawner* CreateFromLocal(TSubclassOf<UK2Node_Variable> NodeClass, UEdGraph* VarContext, FBPVariableDescription const& VarDesc, FProperty* VarProperty, UObject* Outer = nullptr);
+	static UE_API UBlueprintVariableNodeSpawner* CreateFromLocal(TSubclassOf<UK2Node_Variable> NodeClass, UEdGraph* VarContext, FBPVariableDescription const& VarDesc, FProperty* VarProperty, UObject* Outer = nullptr);
 
 	// UBlueprintNodeSpawner interface
-	virtual void Prime() override;
-	virtual FBlueprintNodeSignature GetSpawnerSignature() const override;
-	virtual FBlueprintActionUiSpec GetUiSpec(FBlueprintActionContext const& Context, FBindingSet const& Bindings) const override;
-	virtual UEdGraphNode* Invoke(UEdGraph* ParentGraph, FBindingSet const& Bindings, FVector2D const Location) const override;
+	UE_API virtual void Prime() override;
+	UE_API virtual FBlueprintNodeSignature GetSpawnerSignature() const override;
+	UE_API virtual FBlueprintActionUiSpec GetUiSpec(FBlueprintActionContext const& Context, FBindingSet const& Bindings) const override;
+	UE_API virtual UEdGraphNode* Invoke(UEdGraph* ParentGraph, FBindingSet const& Bindings, FVector2D const Location) const override;
 	// End UBlueprintNodeSpawner interface
 	
 	/**
 	 * @return True if this is a user-created local variable
 	 */
-	bool IsUserLocalVariable() const;
+	UE_API bool IsUserLocalVariable() const;
 
 	/**
 	 * Since this spawner can wrap both local and member variables, we use this
@@ -85,7 +87,7 @@ public:
 	 * 
 	 * @return True if this spawner wraps a local variable, false if not.
 	 */
-	bool IsLocalVariable() const;
+	UE_API bool IsLocalVariable() const;
 
 	/**
 	 * If this is a local variable, then this will return the UEdGraph that it
@@ -94,7 +96,7 @@ public:
 	 * 
 	 * @return Null if the variable is no longer valid, otherwise an object pointer to the variable's owner.
 	 */
-	FFieldVariant GetVarOuter() const;
+	UE_API FFieldVariant GetVarOuter() const;
 
 	/**
 	 * Accessor to the variable's property. Will be null if this is for a local 
@@ -102,7 +104,7 @@ public:
 	 * 
 	 * @return Null if this wraps a local variable (or if the variable property is stale), otherwise the property this was initialized with. 
 	 */
-	FProperty const* GetVarProperty() const;
+	UE_API FProperty const* GetVarProperty() const;
 
 	/**
 	 * Utility function for easily accessing the variable's type (needs to pull
@@ -111,7 +113,7 @@ public:
 	 * 
 	 * @return A struct detailing the wrapped variable's type.
 	 */
-	FEdGraphPinType GetVarType() const;
+	UE_API FEdGraphPinType GetVarType() const;
 
 private:
 	/**
@@ -121,7 +123,7 @@ private:
 	 * 
 	 * @return A friendly, user presentable, name for the variable that this wraps. 
 	 */
-	FText GetVariableName() const;
+	UE_API FText GetVariableName() const;
 
 	/** The graph that the local variable belongs to (if this is a local variable spawner). */
 	UPROPERTY()
@@ -130,4 +132,6 @@ private:
 	/** */
 	UPROPERTY()
 	FBPVariableDescription LocalVarDesc;
-};
\ No newline at end of file
+};
+
+#undef UE_API
diff --git a/Engine/Source/Editor/BlueprintGraph/Classes/EdGraphSchema_K2.h b/Engine/Source/Editor/BlueprintGraph/Classes/EdGraphSchema_K2.h
--- a/Engine/Source/Editor/BlueprintGraph/Classes/EdGraphSchema_K2.h
+++ b/Engine/Source/Editor/BlueprintGraph/Classes/EdGraphSchema_K2.h
@@ -33,6 +33,8 @@
 
 #include "EdGraphSchema_K2.generated.h"
 
+#define UE_API BLUEPRINTGRAPH_API
+
 class AActor;
 class FKismetCompilerContext;
 class FMenuBuilder;
@@ -49,7 +51,7 @@ template <typename T> struct TObjectPtr;
 
 /** Reference to an structure (only used in 'docked' palette) */
 USTRUCT()
-struct BLUEPRINTGRAPH_API FEdGraphSchemaAction_K2Struct : public FEdGraphSchemaAction
+struct FEdGraphSchemaAction_K2Struct : public FEdGraphSchemaAction
 {
 	GENERATED_USTRUCT_BODY()
 
@@ -84,226 +86,226 @@ struct BLUEPRINTGRAPH_API FEdGraphSchemaAction_K2Struct : public FEdGraphSchemaA
 };
 
 // Constants used for metadata, etc... in blueprints
-struct BLUEPRINTGRAPH_API FBlueprintMetadata
+struct FBlueprintMetadata
 {
 public:
 	// Struct/Enum/Class:
 	// If true, this class, struct, or enum is a valid type for use as a variable in a blueprint
-	static const FName MD_AllowableBlueprintVariableType;
+	static UE_API const FName MD_AllowableBlueprintVariableType;
 
 	// If true, this class, struct, or enum is not valid for use as a variable in a blueprint
-	static const FName MD_NotAllowableBlueprintVariableType;
+	static UE_API const FName MD_NotAllowableBlueprintVariableType;
 
 	// Class:
 	// [ClassMetadata] If present, the component class can be spawned by a blueprint
-	static const FName MD_BlueprintSpawnableComponent;
+	static UE_API const FName MD_BlueprintSpawnableComponent;
 
 	/** If true, the class will be usable as a base for blueprints */
-	static const FName MD_IsBlueprintBase;
+	static UE_API const FName MD_IsBlueprintBase;
 	
 	/** A listing of classes that this class is accessible from (and only those classes, if present).  Note that this determines the GRAPH CONTEXTS in which the node cannot be placed (e.g. right click menu, palette), and does NOT control menus when dragging off of a context pin (i.e. contextual drag) */
-	static const FName MD_RestrictedToClasses;
+	static UE_API const FName MD_RestrictedToClasses;
 
 	/// [ClassMetadata] Used for Actor and Component classes. If the native class cannot tick, Blueprint generated classes based this Actor or Component can have bCanEverTick flag overridden even if bCanBlueprintsTickByDefault is false.
-	static const FName MD_ChildCanTick;
+	static UE_API const FName MD_ChildCanTick;
 
 	/// [ClassMetadata] Used for Actor and Component classes. If the native class cannot tick, Blueprint generated classes based this Actor or Component can never tick even if bCanBlueprintsTickByDefault is true.
-	static const FName MD_ChildCannotTick;
+	static UE_API const FName MD_ChildCannotTick;
 
 	/// [ClassMetadata] Used to make the first subclass of a class ignore all inherited showCategories and hideCategories commands
-	static const FName MD_IgnoreCategoryKeywordsInSubclasses;
+	static UE_API const FName MD_IgnoreCategoryKeywordsInSubclasses;
 
 	/** Specifies which struct implements the custom thunk functions for this class */
-	static const FName MD_CustomThunkTemplates;
+	static UE_API const FName MD_CustomThunkTemplates;
 
 	//    function metadata
 	/** Specifies a UFUNCTION as Kismet protected, which can only be called from itself */
-	static const FName MD_Protected;
+	static UE_API const FName MD_Protected;
 
 	/** Marks a UFUNCTION as latent execution */
-	static const FName MD_Latent;
+	static UE_API const FName MD_Latent;
 
 	/** Indicates that the UFUNCTION implements its own thunk function */
-	static const FName MD_CustomThunk;
+	static UE_API const FName MD_CustomThunk;
 
 	/** Marks a UFUNCTION as accepting variadic arguments */
-	static const FName MD_Variadic;
+	static UE_API const FName MD_Variadic;
 
 	/** Marks a UFUNCTION as unsafe for use in the UCS, which prevents it from being called from the UCS.  Useful for things that spawn actors, etc that should never happen in the UCS */
-	static const FName MD_UnsafeForConstructionScripts;
+	static UE_API const FName MD_UnsafeForConstructionScripts;
 
 	// The category that a function appears under in the palette
-	static const FName MD_FunctionCategory;
+	static UE_API const FName MD_FunctionCategory;
 
 	// [FunctionMetadata] Indicates that the function is deprecated
-	static const FName MD_DeprecatedFunction;
+	static UE_API const FName MD_DeprecatedFunction;
 
 	// [FunctionMetadata] Indicates that the property is deprecated
-	static const FName MD_DeprecatedProperty;
+	static UE_API const FName MD_DeprecatedProperty;
 
 	// [FunctionMetadata] Supplies the custom message to use for deprecation
-	static const FName MD_DeprecationMessage;
+	static UE_API const FName MD_DeprecationMessage;
 
 	// [FunctionMetadata] Indicates that the function should be drawn as a compact node with the specified body title
-	static const FName MD_CompactNodeTitle;
+	static UE_API const FName MD_CompactNodeTitle;
 
 	// [FunctionMetadata] Indicates that the function should be drawn with this title over the function name
-	static const FName MD_DisplayName;
+	static UE_API const FName MD_DisplayName;
 
 	// [FunctionMetadata] Indicates the display name of the return value pin
-	static const FName MD_ReturnDisplayName;
+	static UE_API const FName MD_ReturnDisplayName;
 
 	// [FunctionMetadata] Indicates that a particular function parameter is for internal use only, which means it will be both hidden and not connectible.
-	static const FName MD_InternalUseParam;
+	static UE_API const FName MD_InternalUseParam;
 
 	// [FunctionMetadata] Indicates that the function should appear as blueprint function even if it doesn't return a value.
-	static const FName MD_ForceAsFunction;
+	static UE_API const FName MD_ForceAsFunction;
 
 	// [FunctionMetadata] Indicates that the function should be ignored when considered for blueprint type promotion
-	static const FName MD_IgnoreTypePromotion;
+	static UE_API const FName MD_IgnoreTypePromotion;
 
 	//    property metadata
 
 	/** UPROPERTY will be exposed on "Spawn Blueprint" nodes as an input  */
-	static const FName MD_ExposeOnSpawn;
+	static UE_API const FName MD_ExposeOnSpawn;
 
 	/** UPROPERTY uses the specified function as a getter rather than reading from the property directly */
-	static const FName MD_PropertyGetFunction;
+	static UE_API const FName MD_PropertyGetFunction;
 
 	/** UPROPERTY uses the specified function as a setter rather than writing to the property directly */
-	static const FName MD_PropertySetFunction;
+	static UE_API const FName MD_PropertySetFunction;
 
 	/** UPROPERTY cannot be modified by other blueprints */
-	static const FName MD_Private;
+	static UE_API const FName MD_Private;
 
 	/** If true, the self pin should not be shown or connectable regardless of purity, const, etc. similar to InternalUseParam */
-	static const FName MD_HideSelfPin;
+	static UE_API const FName MD_HideSelfPin;
 
 	/** If true, the specified UObject parameter will default to "self" if nothing is connected */
-	static const FName MD_DefaultToSelf;
+	static UE_API const FName MD_DefaultToSelf;
 
 	/** The specified parameter should be used as the context object when retrieving a UWorld pointer (implies hidden and default-to-self) */
-	static const FName MD_WorldContext;
+	static UE_API const FName MD_WorldContext;
 
 	/** For functions that have the MD_WorldContext metadata but are safe to be called from contexts that do not have the ability to provide the world context (either through GetWorld() or ShowWorldContextPin class metadata */
-	static const FName MD_CallableWithoutWorldContext;
+	static UE_API const FName MD_CallableWithoutWorldContext;
 
 	/** For functions that should be compiled in development mode only */
-	static const FName MD_DevelopmentOnly;
+	static UE_API const FName MD_DevelopmentOnly;
 
 	/** If true, an unconnected pin will generate a UPROPERTY under the hood to connect as the input, which will be set to the literal value for the pin.  Only valid for reference parameters. */
-	static const FName MD_AutoCreateRefTerm;
+	static UE_API const FName MD_AutoCreateRefTerm;
 
 	UE_DEPRECATED(5.6, "Use MD_HidePinAssetPicker instead.")
-	static const FName MD_HideAssetPicker;
+	static UE_API const FName MD_HideAssetPicker;
 
 	/** For functions, the specified parameter should hide the asset picker on the pin, even if it is a valid UObject Asset. */
 	/** For properties, forcibly hide the asset picker for the pin corresponding to the property. */
-	static const FName MD_HidePinAssetPicker;
+	static UE_API const FName MD_HidePinAssetPicker;
 
 	/** If true, the hidden world context pin will be visible when the function is placed in a child blueprint of the class. */
-	static const FName MD_ShowWorldContextPin;
+	static UE_API const FName MD_ShowWorldContextPin;
 
 	/** Comma delimited list of pins that should be hidden on this function. */
-	static const FName MD_HidePin;
+	static UE_API const FName MD_HidePin;
 
-	static const FName MD_BlueprintInternalUseOnly;
-	static const FName MD_BlueprintInternalUseOnlyHierarchical;
-	static const FName MD_NeedsLatentFixup;
+	static UE_API const FName MD_BlueprintInternalUseOnly;
+	static UE_API const FName MD_BlueprintInternalUseOnlyHierarchical;
+	static UE_API const FName MD_NeedsLatentFixup;
 
-	static const FName MD_LatentInfo;
-	static const FName MD_LatentCallbackTarget;
+	static UE_API const FName MD_LatentInfo;
+	static UE_API const FName MD_LatentCallbackTarget;
 
 	/** If true, properties defined in the C++ private scope will be accessible to blueprints */
-	static const FName MD_AllowPrivateAccess;
+	static UE_API const FName MD_AllowPrivateAccess;
 
 	/** Categories of functions to expose on this property */
-	static const FName MD_ExposeFunctionCategories;
+	static UE_API const FName MD_ExposeFunctionCategories;
 
 	// [InterfaceMetadata]
-	static const FName MD_CannotImplementInterfaceInBlueprint;
-	static const FName MD_ProhibitedInterfaces;
+	static UE_API const FName MD_CannotImplementInterfaceInBlueprint;
+	static UE_API const FName MD_ProhibitedInterfaces;
 
 	/** Keywords used when searching for functions */
-	static const FName MD_FunctionKeywords;
+	static UE_API const FName MD_FunctionKeywords;
 
 	/** Indicates that during compile we want to create multiple exec pins from an enum param */
-	static const FName MD_ExpandEnumAsExecs;
+	static UE_API const FName MD_ExpandEnumAsExecs;
 
 	/** Synonym for MD_ExpandEnumAsExecs */
-	static const FName MD_ExpandBoolAsExecs;
+	static UE_API const FName MD_ExpandBoolAsExecs;
 
-	static const FName MD_CommutativeAssociativeBinaryOperator;
+	static UE_API const FName MD_CommutativeAssociativeBinaryOperator;
 
 	/** Metadata string that indicates to use the MaterialParameterCollectionFunction node. */
-	static const FName MD_MaterialParameterCollectionFunction;
+	static UE_API const FName MD_MaterialParameterCollectionFunction;
 
 	/** Metadata string that sets the tooltip */
-	static const FName MD_Tooltip;
+	static UE_API const FName MD_Tooltip;
 
 	/** Metadata string that indicates the specified event can be triggered in editor */
-	static const FName MD_CallInEditor;
+	static UE_API const FName MD_CallInEditor;
 
 	/** Metadata to identify an DataTable Pin. Depending on which DataTable is selected, we display different RowName options */
-	static const FName MD_DataTablePin;
+	static UE_API const FName MD_DataTablePin;
 
 	/** Metadata that flags make/break functions for specific struct types. */
-	static const FName MD_NativeMakeFunction;
-	static const FName MD_NativeBreakFunction;
-	static const FName MD_NativeDisableSplitPin;
+	static UE_API const FName MD_NativeMakeFunction;
+	static UE_API const FName MD_NativeBreakFunction;
+	static UE_API const FName MD_NativeDisableSplitPin;
 
 	/** Metadata that flags function params that govern what type of object the function returns */
-	static const FName MD_DynamicOutputType;
+	static UE_API const FName MD_DynamicOutputType;
 	/** Metadata that flags the function output param that will be controlled by the "MD_DynamicOutputType" pin */
-	static const FName MD_DynamicOutputParam;
+	static UE_API const FName MD_DynamicOutputParam;
 
-	static const FName MD_CustomStructureParam;
+	static UE_API const FName MD_CustomStructureParam;
 
-	static const FName MD_ArrayParam;
-	static const FName MD_ArrayDependentParam;
+	static UE_API const FName MD_ArrayParam;
+	static UE_API const FName MD_ArrayDependentParam;
 
 	/** Metadata that flags TSet parameters that will have their type determined at blueprint compile time */
-	static const FName MD_SetParam;
+	static UE_API const FName MD_SetParam;
 
 	/** Metadata that flags TMap function parameters that will have their type determined at blueprint compile time */
-	static const FName MD_MapParam;
-	static const FName MD_MapKeyParam;
-	static const FName MD_MapValueParam;
+	static UE_API const FName MD_MapParam;
+	static UE_API const FName MD_MapKeyParam;
+	static UE_API const FName MD_MapValueParam;
 
 	/** Metadata that identifies an integral property as a bitmask. */
-	static const FName MD_Bitmask;
+	static UE_API const FName MD_Bitmask;
 	/** Metadata that associates a bitmask property with a bitflag enum. */
-	static const FName MD_BitmaskEnum;
+	static UE_API const FName MD_BitmaskEnum;
 	/** Metadata that identifies an enum as a set of explicitly-named bitflags. */
-	static const FName MD_Bitflags;
+	static UE_API const FName MD_Bitflags;
 	/** Metadata that signals to the editor that enum values correspond to mask values instead of bitshift (index) values. */
-	static const FName MD_UseEnumValuesAsMaskValuesInEditor;
+	static UE_API const FName MD_UseEnumValuesAsMaskValuesInEditor;
 	
 	/** Stub function used internally by animation blueprints */
-	static const FName MD_AnimBlueprintFunction;
+	static UE_API const FName MD_AnimBlueprintFunction;
 
 	/** Metadata that should be used with UPARAM to specify whether a TSubclassOf argument should use the DisplayName of classes */
-	static const FName MD_ShowDisplayNames;
+	static UE_API const FName MD_ShowDisplayNames;
 	
 	/** Metadata that should be used with UPARAM to specify whether a TSubclassOf argument allows abstract classes */
-	static const FName MD_AllowAbstractClasses;
+	static UE_API const FName MD_AllowAbstractClasses;
 	/** Metadata that should be used with UPARAM to specify which classes are allowed in selection */
-	static const FName MD_AllowedClasses;
+	static UE_API const FName MD_AllowedClasses;
 
 	/** Metadata that should be used with UPARAM to specify a function name that generates a list of available values */
-	static const FName MD_GetOptions;
+	static UE_API const FName MD_GetOptions;
 
 	/** Namespace into which a type can be optionally defined; if empty or not set, the type will belong to the global namespace (default). */
-	static const FName MD_Namespace;
+	static UE_API const FName MD_Namespace;
 
 	/** Function or class marked as thread-safe. Opts class/function compilation into thread-safety checks. */
-	static const FName MD_ThreadSafe;
+	static UE_API const FName MD_ThreadSafe;
 
 	/** Function marked as explicitly not thread-safe. Opts function out of class-level thread-safety checks. */
-	static const FName MD_NotThreadSafe;
+	static UE_API const FName MD_NotThreadSafe;
 
 	/** Metadata to add the property or function to the FieldNotification system. */
-	static const FName MD_FieldNotify;
+	static UE_API const FName MD_FieldNotify;
 	
 private:
 	// This class should never be instantiated
@@ -357,65 +359,65 @@ enum class ETypeTreeFilter : uint8
 
 ENUM_CLASS_FLAGS(ETypeTreeFilter);
 
-UCLASS(config=Editor)
-class BLUEPRINTGRAPH_API UEdGraphSchema_K2 : public UEdGraphSchema
+UCLASS(MinimalAPI, config=Editor)
+class UEdGraphSchema_K2 : public UEdGraphSchema
 {
 	GENERATED_UCLASS_BODY()
 
 	// Allowable PinType.PinCategory values
-	static const FName PC_Exec;
-	static const FName PC_Boolean;
-	static const FName PC_Byte;
-	static const FName PC_Class;    // SubCategoryObject is the MetaClass of the Class passed thru this pin, or SubCategory can be 'self'. The DefaultValue string should always be empty, use DefaultObject.
-	static const FName PC_SoftClass;
-	static const FName PC_Int;
-	static const FName PC_Int64;
-	static const FName PC_Float;
-	static const FName PC_Double;
-	static const FName PC_Real;
-	static const FName PC_Name;
-	static const FName PC_Delegate;    // SubCategoryObject is the UFunction of the delegate signature
-	static const FName PC_MCDelegate;  // SubCategoryObject is the UFunction of the delegate signature
-	static const FName PC_Object;    // SubCategoryObject is the Class of the object passed thru this pin, or SubCategory can be 'self'. The DefaultValue string should always be empty, use DefaultObject.
-	static const FName PC_Interface;	// SubCategoryObject is the Class of the object passed thru this pin.
-	static const FName PC_SoftObject;		// SubCategoryObject is the Class of the AssetPtr passed thru this pin.
-	static const FName PC_String;
-	static const FName PC_Text;
-	static const FName PC_Struct;    // SubCategoryObject is the ScriptStruct of the struct passed thru this pin, 'self' is not a valid SubCategory. DefaultObject should always be empty, the DefaultValue string may be used for supported structs.
-	static const FName PC_Wildcard;    // Special matching rules are imposed by the node itself
-	static const FName PC_Enum;    // SubCategoryObject is the UEnum object passed thru this pin
-	static const FName PC_FieldPath;		// SubCategoryObject is the Class of the property passed thru this pin.
+	static UE_API const FName PC_Exec;
+	static UE_API const FName PC_Boolean;
+	static UE_API const FName PC_Byte;
+	static UE_API const FName PC_Class;    // SubCategoryObject is the MetaClass of the Class passed thru this pin, or SubCategory can be 'self'. The DefaultValue string should always be empty, use DefaultObject.
+	static UE_API const FName PC_SoftClass;
+	static UE_API const FName PC_Int;
+	static UE_API const FName PC_Int64;
+	static UE_API const FName PC_Float;
+	static UE_API const FName PC_Double;
+	static UE_API const FName PC_Real;
+	static UE_API const FName PC_Name;
+	static UE_API const FName PC_Delegate;    // SubCategoryObject is the UFunction of the delegate signature
+	static UE_API const FName PC_MCDelegate;  // SubCategoryObject is the UFunction of the delegate signature
+	static UE_API const FName PC_Object;    // SubCategoryObject is the Class of the object passed thru this pin, or SubCategory can be 'self'. The DefaultValue string should always be empty, use DefaultObject.
+	static UE_API const FName PC_Interface;	// SubCategoryObject is the Class of the object passed thru this pin.
+	static UE_API const FName PC_SoftObject;		// SubCategoryObject is the Class of the AssetPtr passed thru this pin.
+	static UE_API const FName PC_String;
+	static UE_API const FName PC_Text;
+	static UE_API const FName PC_Struct;    // SubCategoryObject is the ScriptStruct of the struct passed thru this pin, 'self' is not a valid SubCategory. DefaultObject should always be empty, the DefaultValue string may be used for supported structs.
+	static UE_API const FName PC_Wildcard;    // Special matching rules are imposed by the node itself
+	static UE_API const FName PC_Enum;    // SubCategoryObject is the UEnum object passed thru this pin
+	static UE_API const FName PC_FieldPath;		// SubCategoryObject is the Class of the property passed thru this pin.
 
 	// Common PinType.PinSubCategory values
-	static const FName PSC_Self;    // Category=PC_Object or PC_Class, indicates the class being compiled
+	static UE_API const FName PSC_Self;    // Category=PC_Object or PC_Class, indicates the class being compiled
 
-	static const FName PSC_Index;	// Category=PC_Wildcard, indicates the wildcard will only accept Int, Bool, Byte and Enum pins (used when a pin represents indexing a list)
-	static const FName PSC_Bitmask;	// Category=PC_Byte or PC_Int, indicates that the pin represents a bitmask field. SubCategoryObject is either NULL or the UEnum object to which the bitmap is linked for bitflag name specification.
+	static UE_API const FName PSC_Index;	// Category=PC_Wildcard, indicates the wildcard will only accept Int, Bool, Byte and Enum pins (used when a pin represents indexing a list)
+	static UE_API const FName PSC_Bitmask;	// Category=PC_Byte or PC_Int, indicates that the pin represents a bitmask field. SubCategoryObject is either NULL or the UEnum object to which the bitmap is linked for bitflag name specification.
 
 	// Pin names that have special meaning and required types in some contexts (depending on the node type)
-	static const FName PN_Execute;    // Category=PC_Exec, singleton, input
-	static const FName PN_Then;    // Category=PC_Exec, singleton, output
-	static const FName PN_Completed;    // Category=PC_Exec, singleton, output
-	static const FName PN_DelegateEntry;    // Category=PC_Exec, singleton, output; entry point for a dynamically bound delegate
-	static const FName PN_EntryPoint;	// entry point to the ubergraph
-	static const FName PN_Self;    // Category=PC_Object, singleton, input
-	static const FName PN_Else;    // Category=PC_Exec, singleton, output
-	static const FName PN_Loop;    // Category=PC_Exec, singleton, output
-	static const FName PN_After;    // Category=PC_Exec, singleton, output
-	static const FName PN_ReturnValue;		// Category=PC_Object, singleton, output
-	static const FName PN_ObjectToCast;    // Category=PC_Object, singleton, input
-	static const FName PN_Condition;    // Category=PC_Boolean, singleton, input
-	static const FName PN_Start;    // Category=PC_Int, singleton, input
-	static const FName PN_Stop;    // Category=PC_Int, singleton, input
-	static const FName PN_Index;    // Category=PC_Int, singleton, output
-	static const FName PN_Item;    // Category=PC_Int, singleton, output
-	static const FName PN_CastSucceeded;    // Category=PC_Exec, singleton, output
-	static const FName PN_CastFailed;    // Category=PC_Exec, singleton, output
-	static const FString PN_CastedValuePrefix;    // Category=PC_Object, singleton, output; actual pin name varies depending on the type to be casted to, this is just a prefix
+	static UE_API const FName PN_Execute;    // Category=PC_Exec, singleton, input
+	static UE_API const FName PN_Then;    // Category=PC_Exec, singleton, output
+	static UE_API const FName PN_Completed;    // Category=PC_Exec, singleton, output
+	static UE_API const FName PN_DelegateEntry;    // Category=PC_Exec, singleton, output; entry point for a dynamically bound delegate
+	static UE_API const FName PN_EntryPoint;	// entry point to the ubergraph
+	static UE_API const FName PN_Self;    // Category=PC_Object, singleton, input
+	static UE_API const FName PN_Else;    // Category=PC_Exec, singleton, output
+	static UE_API const FName PN_Loop;    // Category=PC_Exec, singleton, output
+	static UE_API const FName PN_After;    // Category=PC_Exec, singleton, output
+	static UE_API const FName PN_ReturnValue;		// Category=PC_Object, singleton, output
+	static UE_API const FName PN_ObjectToCast;    // Category=PC_Object, singleton, input
+	static UE_API const FName PN_Condition;    // Category=PC_Boolean, singleton, input
+	static UE_API const FName PN_Start;    // Category=PC_Int, singleton, input
+	static UE_API const FName PN_Stop;    // Category=PC_Int, singleton, input
+	static UE_API const FName PN_Index;    // Category=PC_Int, singleton, output
+	static UE_API const FName PN_Item;    // Category=PC_Int, singleton, output
+	static UE_API const FName PN_CastSucceeded;    // Category=PC_Exec, singleton, output
+	static UE_API const FName PN_CastFailed;    // Category=PC_Exec, singleton, output
+	static UE_API const FString PN_CastedValuePrefix;    // Category=PC_Object, singleton, output; actual pin name varies depending on the type to be casted to, this is just a prefix
 
 	// construction script function names
-	static const FName FN_UserConstructionScript;
-	static const FName FN_ExecuteUbergraphBase;
+	static UE_API const FName FN_UserConstructionScript;
+	static UE_API const FName FN_ExecuteUbergraphBase;
 
 	// metadata keys
 
@@ -423,23 +425,23 @@ class BLUEPRINTGRAPH_API UEdGraphSchema_K2 : public UEdGraphSchema
 
 
 	// graph names
-	static const FName GN_EventGraph;
-	static const FName GN_AnimGraph;
+	static UE_API const FName GN_EventGraph;
+	static UE_API const FName GN_AnimGraph;
 
 	// variable names
-	static const FText VR_DefaultCategory;
+	static UE_API const FText VR_DefaultCategory;
 
 	// action grouping values
-	static const int32 AG_LevelReference;
+	static UE_API const int32 AG_LevelReference;
 
 	// Somewhat hacky mechanism to prevent tooltips created for pins from including the display name and type when generating BP API documentation
-	static bool bGeneratingDocumentation;
+	static UE_API bool bGeneratingDocumentation;
 
 	// ID for checking dirty status of node titles against, increases every compile
-	static int32 CurrentCacheRefreshID;
+	static UE_API int32 CurrentCacheRefreshID;
 
 	// Pin Selector category for all object types
-	static const FName AllObjectTypes;
+	static UE_API const FName AllObjectTypes;
 
 	UPROPERTY(globalconfig)
 	TArray<FBlueprintCallableFunctionRedirect> EditoronlyBPFunctionRedirects;
@@ -449,7 +451,7 @@ public:
 	//////////////////////////////////////////////////////////////////////////
 	// FPinTypeInfo
 	/** Class used for creating type tree selection info, which aggregates the various PC_* and PinSubtypes in the schema into a heirarchy */
-	class BLUEPRINTGRAPH_API FPinTypeTreeInfo
+	class FPinTypeTreeInfo
 	{
 	private:
 		/** The pin type corresponding to the schema type */
@@ -475,20 +477,20 @@ public:
 		FText Tooltip;
 
 	public:
-		const FEdGraphPinType& GetPinType(bool bForceLoadedSubCategoryObject);
+		UE_API const FEdGraphPinType& GetPinType(bool bForceLoadedSubCategoryObject);
 		const FEdGraphPinType& GetPinTypeNoResolve() const { return PinType; }
 		void SetPinSubTypeCategory(const FName SubCategory)
 		{
 			PinType.PinSubCategory = SubCategory;
 		}
 
-		FPinTypeTreeInfo(const FText& InFriendlyName, const FName CategoryName, const UEdGraphSchema_K2* Schema, const FText& InTooltip, bool bInReadOnly = false);
-		FPinTypeTreeInfo(const FName CategoryName, UObject* SubCategoryObject, const FText& InTooltip, bool bInReadOnly = false, uint8 InPossibleObjectReferenceTypes = 0);
-		FPinTypeTreeInfo(const FText& InFriendlyName, const FName CategoryName, const FAssetData& AssetData, const FText& InTooltip, bool bInReadOnly = false, uint8 InPossibleObjectReferenceTypes = 0);
-		FPinTypeTreeInfo(TSharedPtr<FPinTypeTreeInfo> InInfo);
+		UE_API FPinTypeTreeInfo(const FText& InFriendlyName, const FName CategoryName, const UEdGraphSchema_K2* Schema, const FText& InTooltip, bool bInReadOnly = false);
+		UE_API FPinTypeTreeInfo(const FName CategoryName, UObject* SubCategoryObject, const FText& InTooltip, bool bInReadOnly = false, uint8 InPossibleObjectReferenceTypes = 0);
+		UE_API FPinTypeTreeInfo(const FText& InFriendlyName, const FName CategoryName, const FAssetData& AssetData, const FText& InTooltip, bool bInReadOnly = false, uint8 InPossibleObjectReferenceTypes = 0);
+		UE_API FPinTypeTreeInfo(TSharedPtr<FPinTypeTreeInfo> InInfo);
 		
 		/** Returns a succinct menu description of this type */
-		const FText& GetDescription() const;
+		UE_API const FText& GetDescription() const;
 		
 		/** Returns the localized description as a string, for faster sorting */
 		const FString& GetCachedDescriptionString() const
@@ -516,7 +518,7 @@ public:
 			return PossibleObjectReferenceTypes;
 		}
 
-		const FAssetData& GetCachedAssetData() const;
+		UE_API const FAssetData& GetCachedAssetData() const;
 
 	private:
 
@@ -525,74 +527,74 @@ public:
 			, bReadOnly(false)
 		{}
 
-		FText GenerateDescription();
+		UE_API FText GenerateDescription();
 	};
 
 public:
-	void SelectAllNodesInDirection(TEnumAsByte<enum EEdGraphPinDirection> InDirection, UEdGraph* Graph, UEdGraphPin* InGraphPin);
+	UE_API void SelectAllNodesInDirection(TEnumAsByte<enum EEdGraphPinDirection> InDirection, UEdGraph* Graph, UEdGraphPin* InGraphPin);
 
 	//~ Begin EdGraphSchema Interface
-	virtual void GetContextMenuActions(class UToolMenu* Menu, class UGraphNodeContextMenuContext* Context) const override;
-	virtual const FPinConnectionResponse CanCreateConnection(const UEdGraphPin* A, const UEdGraphPin* B) const override;
-	virtual bool TryCreateConnection(UEdGraphPin* A, UEdGraphPin* B) const override;
-	virtual bool CreateAutomaticConversionNodeAndConnections(UEdGraphPin* A, UEdGraphPin* B) const override;
-	virtual bool CreatePromotedConnection(UEdGraphPin* A, UEdGraphPin* B) const override;
-	virtual FString IsPinDefaultValid(const UEdGraphPin* Pin, const FString& NewDefaultValue, TObjectPtr<UObject> NewDefaultObject, const FText& InNewDefaultText) const override;
-	virtual bool DoesSupportPinWatching() const	override;
-	virtual bool IsPinBeingWatched(UEdGraphPin const* Pin) const override;
-	virtual void ClearPinWatch(UEdGraphPin const* Pin) const override;
-	virtual void TrySetDefaultValue(UEdGraphPin& Pin, const FString& NewDefaultValue, bool bMarkAsModified = true) const override;
-	virtual void TrySetDefaultObject(UEdGraphPin& Pin, UObject* NewDefaultObject, bool bMarkAsModified = true) const override;
-	virtual void TrySetDefaultText(UEdGraphPin& InPin, const FText& InNewDefaultText, bool bMarkAsModified = true) const override;
-	virtual bool DoesDefaultValueMatchAutogenerated(const UEdGraphPin& InPin) const override;
-	virtual void ResetPinToAutogeneratedDefaultValue(UEdGraphPin* Pin, bool bCallModifyCallbacks = true) const override;
-	virtual bool ShouldHidePinDefaultValue(UEdGraphPin* Pin) const override;
-	virtual bool ShouldShowAssetPickerForPin(UEdGraphPin* Pin) const override;
-	virtual FLinearColor GetPinTypeColor(const FEdGraphPinType& PinType) const override;
-	virtual FLinearColor GetSecondaryPinTypeColor(const FEdGraphPinType& PinType) const override;
-	virtual FText GetPinDisplayName(const UEdGraphPin* Pin) const override;
-	virtual void ConstructBasicPinTooltip(const UEdGraphPin& Pin, const FText& PinDescription, FString& TooltipOut) const override;
-	virtual EGraphType GetGraphType(const UEdGraph* TestEdGraph) const override;
-	virtual bool IsTitleBarPin(const UEdGraphPin& Pin) const override;
-	virtual void BreakNodeLinks(UEdGraphNode& TargetNode) const override;
-	virtual void BreakPinLinks(UEdGraphPin& TargetPin, bool bSendsNodeNotifcation) const override;
-	virtual void BreakSinglePinLink(UEdGraphPin* SourcePin, UEdGraphPin* TargetPin) const override;
-	virtual void ReconstructNode(UEdGraphNode& TargetNode, bool bIsBatchRequest=false) const override;
-	virtual bool CanEncapuslateNode(UEdGraphNode const& TestNode) const override;
-	virtual void HandleGraphBeingDeleted(UEdGraph& GraphBeingRemoved) const override;
-	virtual void GetGraphDisplayInformation(const UEdGraph& Graph, /*out*/ FGraphDisplayInfo& DisplayInfo) const override;
-	virtual void DroppedAssetsOnGraph(const TArray<FAssetData>& Assets, const FVector2f& GraphPosition, UEdGraph* Graph) const override;
-	virtual void DroppedAssetsOnNode(const TArray<FAssetData>& Assets, const FVector2f& GraphPosition, UEdGraphNode* Node) const override;
-	virtual void DroppedAssetsOnPin(const TArray<FAssetData>& Assets, const FVector2f& GraphPosition, UEdGraphPin* Pin) const override;
-	virtual void GetAssetsNodeHoverMessage(const TArray<FAssetData>& Assets, const UEdGraphNode* HoverNode, FString& OutTooltipText, bool& OutOkIcon) const override;
-	virtual void GetAssetsPinHoverMessage(const TArray<FAssetData>& Assets, const UEdGraphPin* HoverPin, FString& OutTooltipText, bool& OutOkIcon) const override;
-	virtual void GetAssetsGraphHoverMessage(const TArray<FAssetData>& Assets, const UEdGraph* HoverGraph, FString& OutTooltipText, bool& OutOkIcon) const override;
-	virtual bool CanDuplicateGraph(UEdGraph* InSourceGraph) const override;
-	virtual UEdGraph* DuplicateGraph(UEdGraph* GraphToDuplicate) const override;
-	virtual UEdGraphNode* CreateSubstituteNode(UEdGraphNode* Node, const UEdGraph* Graph, FObjectInstancingGraph* InstanceGraph, TSet<FName>& InOutExtraNames) const override;
-	virtual int32 GetNodeSelectionCount(const UEdGraph* Graph) const override;
-	virtual TSharedPtr<FEdGraphSchemaAction> GetCreateCommentAction() const override;
-	virtual bool FadeNodeWhenDraggingOffPin(const UEdGraphNode* Node, const UEdGraphPin* Pin) const override;
-	virtual void BackwardCompatibilityNodeConversion(UEdGraph* Graph, bool bOnlySafeChanges) const override;
+	UE_API virtual void GetContextMenuActions(class UToolMenu* Menu, class UGraphNodeContextMenuContext* Context) const override;
+	UE_API virtual const FPinConnectionResponse CanCreateConnection(const UEdGraphPin* A, const UEdGraphPin* B) const override;
+	UE_API virtual bool TryCreateConnection(UEdGraphPin* A, UEdGraphPin* B) const override;
+	UE_API virtual bool CreateAutomaticConversionNodeAndConnections(UEdGraphPin* A, UEdGraphPin* B) const override;
+	UE_API virtual bool CreatePromotedConnection(UEdGraphPin* A, UEdGraphPin* B) const override;
+	UE_API virtual FString IsPinDefaultValid(const UEdGraphPin* Pin, const FString& NewDefaultValue, TObjectPtr<UObject> NewDefaultObject, const FText& InNewDefaultText) const override;
+	UE_API virtual bool DoesSupportPinWatching() const	override;
+	UE_API virtual bool IsPinBeingWatched(UEdGraphPin const* Pin) const override;
+	UE_API virtual void ClearPinWatch(UEdGraphPin const* Pin) const override;
+	UE_API virtual void TrySetDefaultValue(UEdGraphPin& Pin, const FString& NewDefaultValue, bool bMarkAsModified = true) const override;
+	UE_API virtual void TrySetDefaultObject(UEdGraphPin& Pin, UObject* NewDefaultObject, bool bMarkAsModified = true) const override;
+	UE_API virtual void TrySetDefaultText(UEdGraphPin& InPin, const FText& InNewDefaultText, bool bMarkAsModified = true) const override;
+	UE_API virtual bool DoesDefaultValueMatchAutogenerated(const UEdGraphPin& InPin) const override;
+	UE_API virtual void ResetPinToAutogeneratedDefaultValue(UEdGraphPin* Pin, bool bCallModifyCallbacks = true) const override;
+	UE_API virtual bool ShouldHidePinDefaultValue(UEdGraphPin* Pin) const override;
+	UE_API virtual bool ShouldShowAssetPickerForPin(UEdGraphPin* Pin) const override;
+	UE_API virtual FLinearColor GetPinTypeColor(const FEdGraphPinType& PinType) const override;
+	UE_API virtual FLinearColor GetSecondaryPinTypeColor(const FEdGraphPinType& PinType) const override;
+	UE_API virtual FText GetPinDisplayName(const UEdGraphPin* Pin) const override;
+	UE_API virtual void ConstructBasicPinTooltip(const UEdGraphPin& Pin, const FText& PinDescription, FString& TooltipOut) const override;
+	UE_API virtual EGraphType GetGraphType(const UEdGraph* TestEdGraph) const override;
+	UE_API virtual bool IsTitleBarPin(const UEdGraphPin& Pin) const override;
+	UE_API virtual void BreakNodeLinks(UEdGraphNode& TargetNode) const override;
+	UE_API virtual void BreakPinLinks(UEdGraphPin& TargetPin, bool bSendsNodeNotifcation) const override;
+	UE_API virtual void BreakSinglePinLink(UEdGraphPin* SourcePin, UEdGraphPin* TargetPin) const override;
+	UE_API virtual void ReconstructNode(UEdGraphNode& TargetNode, bool bIsBatchRequest=false) const override;
+	UE_API virtual bool CanEncapuslateNode(UEdGraphNode const& TestNode) const override;
+	UE_API virtual void HandleGraphBeingDeleted(UEdGraph& GraphBeingRemoved) const override;
+	UE_API virtual void GetGraphDisplayInformation(const UEdGraph& Graph, /*out*/ FGraphDisplayInfo& DisplayInfo) const override;
+	UE_API virtual void DroppedAssetsOnGraph(const TArray<FAssetData>& Assets, const FVector2f& GraphPosition, UEdGraph* Graph) const override;
+	UE_API virtual void DroppedAssetsOnNode(const TArray<FAssetData>& Assets, const FVector2f& GraphPosition, UEdGraphNode* Node) const override;
+	UE_API virtual void DroppedAssetsOnPin(const TArray<FAssetData>& Assets, const FVector2f& GraphPosition, UEdGraphPin* Pin) const override;
+	UE_API virtual void GetAssetsNodeHoverMessage(const TArray<FAssetData>& Assets, const UEdGraphNode* HoverNode, FString& OutTooltipText, bool& OutOkIcon) const override;
+	UE_API virtual void GetAssetsPinHoverMessage(const TArray<FAssetData>& Assets, const UEdGraphPin* HoverPin, FString& OutTooltipText, bool& OutOkIcon) const override;
+	UE_API virtual void GetAssetsGraphHoverMessage(const TArray<FAssetData>& Assets, const UEdGraph* HoverGraph, FString& OutTooltipText, bool& OutOkIcon) const override;
+	UE_API virtual bool CanDuplicateGraph(UEdGraph* InSourceGraph) const override;
+	UE_API virtual UEdGraph* DuplicateGraph(UEdGraph* GraphToDuplicate) const override;
+	UE_API virtual UEdGraphNode* CreateSubstituteNode(UEdGraphNode* Node, const UEdGraph* Graph, FObjectInstancingGraph* InstanceGraph, TSet<FName>& InOutExtraNames) const override;
+	UE_API virtual int32 GetNodeSelectionCount(const UEdGraph* Graph) const override;
+	UE_API virtual TSharedPtr<FEdGraphSchemaAction> GetCreateCommentAction() const override;
+	UE_API virtual bool FadeNodeWhenDraggingOffPin(const UEdGraphNode* Node, const UEdGraphPin* Pin) const override;
+	UE_API virtual void BackwardCompatibilityNodeConversion(UEdGraph* Graph, bool bOnlySafeChanges) const override;
 	virtual bool ShouldAlwaysPurgeOnModification() const override { return false; }
-	virtual void SplitPin(UEdGraphPin* Pin, bool bNotify = true) const override;
-	virtual void RecombinePin(UEdGraphPin* Pin) const override;
-	virtual void OnPinConnectionDoubleCicked(UEdGraphPin* PinA, UEdGraphPin* PinB, const FVector2f& GraphPosition) const override;
-	virtual UEdGraphPin* DropPinOnNode(UEdGraphNode* InTargetNode, const FName& InSourcePinName, const FEdGraphPinType& InSourcePinType, EEdGraphPinDirection InSourcePinDirection) const override;
-	virtual bool SupportsDropPinOnNode(UEdGraphNode* InTargetNode, const FEdGraphPinType& InSourcePinType, EEdGraphPinDirection InSourcePinDirection, FText& OutErrorMessage) const override;
-	virtual bool IsCacheVisualizationOutOfDate(int32 InVisualizationCacheID) const override;
-	virtual int32 GetCurrentVisualizationCacheID() const override;
-	virtual void ForceVisualizationCacheClear() const override;
-	virtual bool SafeDeleteNodeFromGraph(UEdGraph* Graph, UEdGraphNode* NodeToDelete) const override;
+	UE_API virtual void SplitPin(UEdGraphPin* Pin, bool bNotify = true) const override;
+	UE_API virtual void RecombinePin(UEdGraphPin* Pin) const override;
+	UE_API virtual void OnPinConnectionDoubleCicked(UEdGraphPin* PinA, UEdGraphPin* PinB, const FVector2f& GraphPosition) const override;
+	UE_API virtual UEdGraphPin* DropPinOnNode(UEdGraphNode* InTargetNode, const FName& InSourcePinName, const FEdGraphPinType& InSourcePinType, EEdGraphPinDirection InSourcePinDirection) const override;
+	UE_API virtual bool SupportsDropPinOnNode(UEdGraphNode* InTargetNode, const FEdGraphPinType& InSourcePinType, EEdGraphPinDirection InSourcePinDirection, FText& OutErrorMessage) const override;
+	UE_API virtual bool IsCacheVisualizationOutOfDate(int32 InVisualizationCacheID) const override;
+	UE_API virtual int32 GetCurrentVisualizationCacheID() const override;
+	UE_API virtual void ForceVisualizationCacheClear() const override;
+	UE_API virtual bool SafeDeleteNodeFromGraph(UEdGraph* Graph, UEdGraphNode* NodeToDelete) const override;
 	virtual bool CanVariableBeDropped(UEdGraph* InGraph, FProperty* InVariableToDrop) const override { return true; }
-	virtual bool CanShowDataTooltipForPin(const UEdGraphPin& Pin) const override;
+	UE_API virtual bool CanShowDataTooltipForPin(const UEdGraphPin& Pin) const override;
 
 #if WITH_EDITORONLY_DATA
-	virtual float GetActionFilteredWeight(const FEdGraphSchemaAction& InCurrentAction, const TArray<FString>& InFilterTerms, const TArray<FString>& InSanitizedFilterTerms, const TArray<UEdGraphPin*>& DraggedFromPins) const override;
+	UE_API virtual float GetActionFilteredWeight(const FEdGraphSchemaAction& InCurrentAction, const TArray<FString>& InFilterTerms, const TArray<FString>& InSanitizedFilterTerms, const TArray<UEdGraphPin*>& DraggedFromPins) const override;
 	PRAGMA_DISABLE_DEPRECATION_WARNINGS
 	UE_DEPRECATED(5.5, "FGraphActionListBuilderBase::ActionGroup weight calculation was always performed on the first action, provide that if scoring a FGraphActionListBuilderBase::ActionGroup, or just provide the FEdGraphSchemaAction")
-	virtual float GetActionFilteredWeight(const FGraphActionListBuilderBase::ActionGroup& InCurrentAction, const TArray<FString>& InFilterTerms, const TArray<FString>& InSanitizedFilterTerms, const TArray<UEdGraphPin*>& DraggedFromPins) const override;
-	PRAGMA_ENABLE_DEPRECATION_WARNINGS
+	UE_API virtual float GetActionFilteredWeight(const FGraphActionListBuilderBase::ActionGroup& InCurrentAction, const TArray<FString>& InFilterTerms, const TArray<FString>& InSanitizedFilterTerms, const TArray<UEdGraphPin*>& DraggedFromPins) const override;
+	UE_API PRAGMA_ENABLE_DEPRECATION_WARNINGS
 	virtual FGraphSchemaSearchWeightModifiers GetSearchWeightModifiers() const override;
 #endif // WITH_EDITORONLY_DATA	
 
@@ -621,7 +623,7 @@ public:
 	* @param	InVaraiableSource	The source of the variable
 	* @param	InTargetBlueprint	The blueprint this node will be used on
 	*/
-	static void ConfigureVarNode(class UK2Node_Variable* InVarNode, FName InVariableName, UStruct* InVariableSource, UBlueprint* InTargetBlueprint);
+	static UE_API void ConfigureVarNode(class UK2Node_Variable* InVarNode, FName InVariableName, UStruct* InVariableSource, UBlueprint* InTargetBlueprint);
 
 	/**
 	* Creates a new variable getter node and adds it to ParentGraph
@@ -632,7 +634,7 @@ public:
 	* @param	Source				The source of the variable
 	* @return	A pointer to the newly spawned node
 	*/
-	virtual class UK2Node_VariableGet* SpawnVariableGetNode(const FVector2D GraphPosition, class UEdGraph* ParentGraph, FName VariableName, UStruct* Source) const;
+	UE_API virtual class UK2Node_VariableGet* SpawnVariableGetNode(const FVector2D GraphPosition, class UEdGraph* ParentGraph, FName VariableName, UStruct* Source) const;
 
 	/**
 	* Creates a new variable setter node and adds it to ParentGraph
@@ -643,16 +645,16 @@ public:
 	* @param	Source				The source of the variable
 	* @return	A pointer to the newly spawned node
 	*/
-	virtual class UK2Node_VariableSet* SpawnVariableSetNode(const FVector2D GraphPosition, class UEdGraph* ParentGraph, FName VariableName, UStruct* Source) const;
+	UE_API virtual class UK2Node_VariableSet* SpawnVariableSetNode(const FVector2D GraphPosition, class UEdGraph* ParentGraph, FName VariableName, UStruct* Source) const;
 
 	// Returns whether the supplied Pin is a splittable struct
-	bool PinHasSplittableStructType(const UEdGraphPin* InGraphPin) const;
+	UE_API bool PinHasSplittableStructType(const UEdGraphPin* InGraphPin) const;
 
 	/** Returns true if the pin has a value field that can be edited inline */
-	bool PinDefaultValueIsEditable(const UEdGraphPin& InGraphPin) const;
+	UE_API bool PinDefaultValueIsEditable(const UEdGraphPin& InGraphPin) const;
 
 	/** Returns true if the pin has a custom default string format and it is not safe to use ExportText */
-	bool PinHasCustomDefaultFormat(const UEdGraphPin& InGraphPin) const;
+	UE_API bool PinHasCustomDefaultFormat(const UEdGraphPin& InGraphPin) const;
 
 	struct FCreateSplitPinNodeParams
 	{
@@ -675,19 +677,19 @@ public:
 
 	/** Helper function to create the expansion node.  
 		If the CompilerContext is specified this will be created as an intermediate node */
-	UK2Node* CreateSplitPinNode(UEdGraphPin* Pin, const FCreateSplitPinNodeParams& Params) const;
+	UE_API UK2Node* CreateSplitPinNode(UEdGraphPin* Pin, const FCreateSplitPinNodeParams& Params) const;
 
 	/** Reads in a FString and gets the values of the pin defaults for that type. This can be passed to DefaultValueSimpleValidation to validate. OwningObject can be null */
-	virtual void GetPinDefaultValuesFromString(const FEdGraphPinType& PinType, UObject* OwningObject, const FString& NewValue, FString& UseDefaultValue, TObjectPtr<UObject>& UseDefaultObject, FText& UseDefaultText, bool bPreserveTextIdentity = true) const;
+	UE_API virtual void GetPinDefaultValuesFromString(const FEdGraphPinType& PinType, UObject* OwningObject, const FString& NewValue, FString& UseDefaultValue, TObjectPtr<UObject>& UseDefaultObject, FText& UseDefaultText, bool bPreserveTextIdentity = true) const;
 
 	/** Do validation, that doesn't require a knowledge about actual pin */
-	virtual bool DefaultValueSimpleValidation(const FEdGraphPinType& PinType, const FName PinName, const FString& NewDefaultValue, TObjectPtr<UObject> NewDefaultObject, const FText& InText, FString* OutMsg = nullptr) const;
+	UE_API virtual bool DefaultValueSimpleValidation(const FEdGraphPinType& PinType, const FName PinName, const FString& NewDefaultValue, TObjectPtr<UObject> NewDefaultObject, const FText& InText, FString* OutMsg = nullptr) const;
 
 	/** Returns true if the owning node is a function with AutoCreateRefTerm meta data */
-	static bool IsAutoCreateRefTerm(const UEdGraphPin* Pin);
+	static UE_API bool IsAutoCreateRefTerm(const UEdGraphPin* Pin);
 
 	/** See if a class has any members that are accessible by a blueprint */
-	bool ClassHasBlueprintAccessibleMembers(const UClass* InClass) const;
+	UE_API bool ClassHasBlueprintAccessibleMembers(const UClass* InClass) const;
 
 	/**
 	 * Checks to see if the specified graph is a construction script
@@ -695,7 +697,7 @@ public:
 	 * @param	TestEdGraph		Graph to test
 	 * @return	true if this is a construction script
 	 */
-	static bool IsConstructionScript(const UEdGraph* TestEdGraph);
+	static UE_API bool IsConstructionScript(const UEdGraph* TestEdGraph);
 
 	/**
 	 * Checks to see if the specified graph is a composite graph
@@ -703,7 +705,7 @@ public:
 	 * @param	TestEdGraph		Graph to test
 	 * @return	true if this is a composite graph
 	 */
-	bool IsCompositeGraph(const UEdGraph* TestEdGraph) const;
+	UE_API bool IsCompositeGraph(const UEdGraph* TestEdGraph) const;
 
 	/**
 	 * Checks to see if the specified graph is a const function graph
@@ -712,7 +714,7 @@ public:
 	 * @param	bOutIsEnforcingConstCorrectness (Optional) Whether or not this graph is enforcing const correctness during compilation
 	 * @return	true if this is a const function graph
 	 */
-	bool IsConstFunctionGraph(const UEdGraph* TestEdGraph, bool* bOutIsEnforcingConstCorrectness = nullptr) const;
+	UE_API bool IsConstFunctionGraph(const UEdGraph* TestEdGraph, bool* bOutIsEnforcingConstCorrectness = nullptr) const;
 
 	/**
 	 * Checks to see if the specified graph is a static function graph
@@ -720,7 +722,7 @@ public:
 	 * @param	TestEdGraph		Graph to test
 	 * @return	true if this is a const function graph
 	 */
-	bool IsStaticFunctionGraph(const UEdGraph* TestEdGraph) const;
+	UE_API bool IsStaticFunctionGraph(const UEdGraph* TestEdGraph) const;
 
 	/**
 	 * Checks to see if a pin is an execution pin.
@@ -739,7 +741,7 @@ public:
 	 * @param	Pin	The pin to check.
 	 * @return	true if it is a Self pin.
 	 */
-	virtual bool IsSelfPin(const UEdGraphPin& Pin) const override;
+	UE_API virtual bool IsSelfPin(const UEdGraphPin& Pin) const override;
 
 	/**
 	 * Checks to see if a pin is a meta-pin (either a Self or Exec pin)
@@ -753,7 +755,7 @@ public:
 	}
 
 	/** Is given string a delegate category name ? */
-	virtual bool IsDelegateCategory(const FName Category) const override;
+	UE_API virtual bool IsDelegateCategory(const FName Category) const override;
 
 	/** Returns whether a pin category is compatible with an Index Wildcard (PC_Wildcard and PSC_Index) */
 	inline bool IsIndexWildcardCompatible(const FEdGraphPinType& PinType) const
@@ -814,20 +816,20 @@ public:
 	}
 
 	/** Can Pin be promoted to a variable? */
-	bool CanPromotePinToVariable (const UEdGraphPin& Pin, bool bInToMemberVariable) const;
+	UE_API bool CanPromotePinToVariable (const UEdGraphPin& Pin, bool bInToMemberVariable) const;
 
 	/** Can Pin be split in to its component elements */
-	bool CanSplitStructPin(const UEdGraphPin& Pin) const;
+	UE_API bool CanSplitStructPin(const UEdGraphPin& Pin) const;
 
 	/** Can Pin be recombined back to its original form */
-	bool CanRecombineStructPin(const UEdGraphPin& Pin) const;
+	UE_API bool CanRecombineStructPin(const UEdGraphPin& Pin) const;
 
 	/** 
 	 * Helper function for filling out Category, SubCategory, and SubCategoryObject based on a FProperty 
 	 * 
 	 * @return	true on success, false if the property is unsupported or invalid.
 	 */
-	static bool GetPropertyCategoryInfo(const FProperty* TestProperty, FName& OutCategory, FName& OutSubCategory, UObject*& OutSubCategoryObject, bool& bOutIsWeakPointer);
+	static UE_API bool GetPropertyCategoryInfo(const FProperty* TestProperty, FName& OutCategory, FName& OutSubCategory, UObject*& OutSubCategoryObject, bool& bOutIsWeakPointer);
 
 	/**
 	 * Convert the type of a FProperty to the corresponding pin type.
@@ -837,13 +839,13 @@ public:
 	 *
 	 * @return	true on success, false if the property is unsupported or invalid.
 	 */
-	bool ConvertPropertyToPinType(const FProperty* Property, /*out*/ FEdGraphPinType& TypeOut) const;
+	UE_API bool ConvertPropertyToPinType(const FProperty* Property, /*out*/ FEdGraphPinType& TypeOut) const;
 
 	/** Returns true if the function has wildcard parameters, e.g. uses runtime type information that may require safe failure handling */
-	static bool HasWildcardParams(const UFunction* Function);
+	static UE_API bool HasWildcardParams(const UFunction* Function);
 
 	/** Determines if the specified param property is intended to be used as a wildcard (for custom thunk functions, like in our array library, etc.)*/
-	static bool IsWildcardProperty(const FProperty* ParamProperty);
+	static UE_API bool IsWildcardProperty(const FProperty* ParamProperty);
 
 	/** Flags to indicate different types of blueprint callable functions */
 	enum EFunctionType
@@ -860,10 +862,10 @@ public:
 	 * @param	Function			The function to find a parent function for
 	 * @return	The UFunction parentfunction, if any.
 	 */
-	static UFunction* GetCallableParentFunction(UFunction* Function);
+	static UE_API UFunction* GetCallableParentFunction(UFunction* Function);
 
 	/** Whether or not the specified actor is a valid target for bound events and literal references (in the right level, not a builder brush, etc */
-	bool IsActorValidForLevelScriptRefs(const AActor* TestActor, const UBlueprint* Blueprint) const;
+	UE_API bool IsActorValidForLevelScriptRefs(const AActor* TestActor, const UBlueprint* Blueprint) const;
 
 	/** 
 	 *	Generate a list of replaceable nodes for context menu based on the editor's current selection 
@@ -871,7 +873,7 @@ public:
 	 *	@param	Reference to graph node
 	 *	@param	Reference to context menu builder
 	 */
-	void AddSelectedReplaceableNodes(FToolMenuSection& Section, UBlueprint* Blueprint, const UEdGraphNode* InGraphNode) const;
+	UE_API void AddSelectedReplaceableNodes(FToolMenuSection& Section, UBlueprint* Blueprint, const UEdGraphNode* InGraphNode) const;
 
 	/**
 	 *	Function to replace current graph node reference object with a new object
@@ -879,19 +881,19 @@ public:
 	 *	@param Reference to graph node
 	 *	@param Reference to new reference object
 	 */
-	void ReplaceSelectedNode(UEdGraphNode* SourceNode, AActor* TargetActor);
+	UE_API void ReplaceSelectedNode(UEdGraphNode* SourceNode, AActor* TargetActor);
 
 	/** Returns whether a function is marked 'override' and doesn't have any out parameters */
-	static bool FunctionCanBePlacedAsEvent(const UFunction* InFunction);
+	static UE_API bool FunctionCanBePlacedAsEvent(const UFunction* InFunction);
 
 	/** Can this function be called by kismet delegate */
-	static bool FunctionCanBeUsedInDelegate(const UFunction* InFunction);
+	static UE_API bool FunctionCanBeUsedInDelegate(const UFunction* InFunction);
 
 	/** Can this function be called by kismet code */
-	static bool CanUserKismetCallFunction(const UFunction* Function);
+	static UE_API bool CanUserKismetCallFunction(const UFunction* Function);
 
 	/** Returns if function has output parameter(s) */
-	static bool HasFunctionAnyOutputParameter(const UFunction* Function);
+	static UE_API bool HasFunctionAnyOutputParameter(const UFunction* Function);
 
 	enum EDelegateFilterMode
 	{
@@ -901,32 +903,32 @@ public:
 	};
 
 	/** Can this variable be accessed by kismet code */
-	static bool CanUserKismetAccessVariable(const FProperty* Property, const UClass* InClass, EDelegateFilterMode FilterMode);
+	static UE_API bool CanUserKismetAccessVariable(const FProperty* Property, const UClass* InClass, EDelegateFilterMode FilterMode);
 
 	/** Can this function be overridden by kismet (either placed as event or new function graph created) */
-	static bool CanKismetOverrideFunction(const UFunction* Function);
+	static UE_API bool CanKismetOverrideFunction(const UFunction* Function);
 
 	/** returns friendly signature name if possible or Removes any mangling to get the unmangled signature name of the function */
-	static FText GetFriendlySignatureName(const UFunction* Function);
+	static UE_API FText GetFriendlySignatureName(const UFunction* Function);
 
 	/** Returns true if this enum is safe to be used as a variable in blueprints */
-	static bool IsAllowableBlueprintVariableType(const UEnum* InEnum);
+	static UE_API bool IsAllowableBlueprintVariableType(const UEnum* InEnum);
 
 	/** 
 	 * Returns true if this class is safe to be used as a variable in blueprints
 	 *
 	 * @param	bImpliedBlueprintType	If true, the class is implied to be a blueprint type and will be allowed unless it is specifically forbidden
 	 */
-	static bool IsAllowableBlueprintVariableType(const UClass* InClass, bool bAssumeBlueprintType = false);
+	static UE_API bool IsAllowableBlueprintVariableType(const UClass* InClass, bool bAssumeBlueprintType = false);
 
 	/**
 	 * Returns true if this struct is safe to to be used as a variable in blueprints
 	 *
 	 * @param	bForInternalUse			If true, the struct is for internal usage so BlueprintInternalUseOnly is allowed
 	 */
-	static bool IsAllowableBlueprintVariableType(const UScriptStruct *InStruct, bool bForInternalUse = false);
+	static UE_API bool IsAllowableBlueprintVariableType(const UScriptStruct *InStruct, bool bForInternalUse = false);
 
-	static bool IsPropertyExposedOnSpawn(const FProperty* Property);
+	static UE_API bool IsPropertyExposedOnSpawn(const FProperty* Property);
 
 	/**
 	 * Returns a list of parameters for the function that are specified as automatically emitting terms for unconnected ref parameters in the compiler (MD_AutoCreateRefTerm)
@@ -934,7 +936,7 @@ public:
 	 * @param	Function				The function to check for auto-emitted ref terms on
 	 * @param	AutoEmitParameterNames	(out) Returns an array of param names that should be auto-emitted if nothing is connected
 	 */
-	static void GetAutoEmitTermParameters(const UFunction* Function, TArray<FString>& AutoEmitParameterNames);
+	static UE_API void GetAutoEmitTermParameters(const UFunction* Function, TArray<FString>& AutoEmitParameterNames);
 
 	/**
 	 * Determine if a function has a parameter of a specific type.
@@ -946,7 +948,7 @@ public:
 	 *
 	 * @return	true if at least one parameter is of the correct type and direction.
 	 */
-	bool FunctionHasParamOfType(const UFunction* InFunction, UEdGraph const* InGraph, const FEdGraphPinType& DesiredPinType, bool bWantOutput) const;
+	UE_API bool FunctionHasParamOfType(const UFunction* InFunction, UEdGraph const* InGraph, const FEdGraphPinType& DesiredPinType, bool bWantOutput) const;
 
 	/**
 	 * Add the specified flags to the function entry node of the graph, to make sure they get compiled in to the generated function
@@ -954,7 +956,7 @@ public:
 	 * @param	CurrentGraph	The graph of the function to modify the flags for
 	 * @param	ExtraFlags		The flags to add to the function entry node
 	 */
-	void AddExtraFunctionFlags(const UEdGraph* CurrentGraph, int32 ExtraFlags) const;
+	UE_API void AddExtraFunctionFlags(const UEdGraph* CurrentGraph, int32 ExtraFlags) const;
 
 	/**
 	 * Marks the function entry of a graph as editable via function editor or not-editable
@@ -962,7 +964,7 @@ public:
 	 * @param	CurrentGraph	The graph of the function to modify the entry node for
 	 * @param	bNewEditable	Whether or not the function entry for the graph should be editable via the function editor
 	 */
-	void MarkFunctionEntryAsEditable(const UEdGraph* CurrentGraph, bool bNewEditable) const;
+	UE_API void MarkFunctionEntryAsEditable(const UEdGraph* CurrentGraph, bool bNewEditable) const;
 
 	/** 
 	 * Populate new macro graph with entry and possibly return node
@@ -970,7 +972,7 @@ public:
 	 * @param	Graph			Graph to add the function terminators to
 	 * @param	ContextClass	If specified, the graph terminators will use this class to search for context for signatures (i.e. interface functions)
 	 */
-	virtual void CreateMacroGraphTerminators(UEdGraph& Graph, UClass* Class) const;
+	UE_API virtual void CreateMacroGraphTerminators(UEdGraph& Graph, UClass* Class) const;
 
 	/** 
 	 * Populate new function graph with entry and possibly return node
@@ -978,7 +980,7 @@ public:
 	 * @param	Graph			Graph to add the function terminators to
 	 * @param	ContextClass	If specified, the graph terminators will use this class to search for context for signatures (i.e. interface functions)
 	 */
-	virtual void CreateFunctionGraphTerminators(UEdGraph& Graph, UClass* Class) const;
+	UE_API virtual void CreateFunctionGraphTerminators(UEdGraph& Graph, UClass* Class) const;
 
 	/**
 	 * Populate new function graph with entry and possibly return node
@@ -986,7 +988,7 @@ public:
 	 * @param	Graph			Graph to add the function terminators to
 	 * @param	FunctionSignature	The function signature to mimic when creating the inputs and outputs for the function.
 	 */
-	virtual void CreateFunctionGraphTerminators(UEdGraph& Graph, const UFunction* FunctionSignature) const;
+	UE_API virtual void CreateFunctionGraphTerminators(UEdGraph& Graph, const UFunction* FunctionSignature) const;
 
 	/**
 	 * Converts the type of a property into a fully qualified string (e.g., object'ObjectName').
@@ -995,7 +997,7 @@ public:
 	 *
 	 * @return	The converted type string.
 	 */
-	static FText TypeToText(const FProperty* const Property);
+	static UE_API FText TypeToText(const FProperty* const Property);
 
 	/**
 	* Converts a terminal type into a fully qualified FText (e.g., object'ObjectName').
@@ -1008,7 +1010,7 @@ public:
 	*
 	* @return	The converted type text.
 	*/
-	static FText TerminalTypeToText(const FName Category, const FName SubCategory, UObject* SubCategoryObject, bool bIsWeakPtr);
+	static UE_API FText TerminalTypeToText(const FName Category, const FName SubCategory, UObject* SubCategoryObject, bool bIsWeakPtr);
 
 	/**
 	 * Converts a pin type into a fully qualified FText (e.g., object'ObjectName').
@@ -1017,7 +1019,7 @@ public:
 	 *
 	 * @return	The converted type text.
 	 */
-	static FText TypeToText(const FEdGraphPinType& Type);
+	static UE_API FText TypeToText(const FEdGraphPinType& Type);
 
 	/**
 	 * Returns the FText to use for a given schema category
@@ -1027,7 +1029,7 @@ public:
 	 *
 	 * @return	The text to display for the category.
 	 */
-	static FText GetCategoryText(const FName Category, const bool bForMenu = false);
+	static UE_API FText GetCategoryText(const FName Category, const bool bForMenu = false);
 
 	/**
 	 * Returns the FText to use for a given schema category and subcategory
@@ -1038,7 +1040,7 @@ public:
 	 *
 	 * @return	The text to display for the category.
 	 */
-	static FText GetCategoryText(FName Category, FName SubCategory, bool bForMenu = false);
+	static UE_API FText GetCategoryText(FName Category, FName SubCategory, bool bForMenu = false);
 
 	/**
 	 * Get the type tree for all of the property types valid for this schema
@@ -1046,14 +1048,14 @@ public:
 	 * @param	TypeTree		The array that will contain the type tree hierarchy for this schema upon returning
 	 * @param	TypeTreeFilter	ETypeTreeFilter flags that determine how the TypeTree is populated.
 	 */
-	void GetVariableTypeTree(TArray< TSharedPtr<FPinTypeTreeInfo> >& TypeTree, ETypeTreeFilter TypeTreeFilter = ETypeTreeFilter::None) const;
+	UE_API void GetVariableTypeTree(TArray< TSharedPtr<FPinTypeTreeInfo> >& TypeTree, ETypeTreeFilter TypeTreeFilter = ETypeTreeFilter::None) const;
 
 	/**
 	 * Returns whether or not the specified type has valid subtypes available
 	 *
 	 * @param	Type	The type to check for subtypes
 	 */
-	bool DoesTypeHaveSubtypes( const FName Category ) const;
+	UE_API bool DoesTypeHaveSubtypes( const FName Category ) const;
 
 	/**
 	 * Returns true if the types and directions of two pins are schema compatible. Handles
@@ -1066,7 +1068,7 @@ public:
 	 *
 	 * @return	true if the pin types and directions are compatible.
 	 */
-	virtual bool ArePinsCompatible(const UEdGraphPin* PinA, const UEdGraphPin* PinB, const UClass* CallingContext = NULL, bool bIgnoreArray = false) const override;
+	UE_API virtual bool ArePinsCompatible(const UEdGraphPin* PinA, const UEdGraphPin* PinB, const UClass* CallingContext = NULL, bool bIgnoreArray = false) const override;
 
 	/**
 	 * Returns the connection response for connecting PinA to PinB, which have already been determined to be compatible
@@ -1079,7 +1081,7 @@ public:
 	 *
 	 * @return	The message and action to take on trying to make this connection.
 	 */
-	virtual const FPinConnectionResponse DetermineConnectionResponseOfCompatibleTypedPins(const UEdGraphPin* PinA, const UEdGraphPin* PinB, const UEdGraphPin* InputPin, const UEdGraphPin* OutputPin) const;
+	UE_API virtual const FPinConnectionResponse DetermineConnectionResponseOfCompatibleTypedPins(const UEdGraphPin* PinA, const UEdGraphPin* PinB, const UEdGraphPin* InputPin, const UEdGraphPin* OutputPin) const;
 
 	/**
 	 * Returns true if the two pin types are schema compatible.  Handles outputting a more derived
@@ -1092,7 +1094,7 @@ public:
 	 *
 	 * @return	true if the pin types are compatible.
 	 */
-	virtual bool ArePinTypesCompatible(const FEdGraphPinType& Output, const FEdGraphPinType& Input, const UClass* CallingContext = NULL, bool bIgnoreArray = false) const;
+	UE_API virtual bool ArePinTypesCompatible(const FEdGraphPinType& Output, const FEdGraphPinType& Input, const UClass* CallingContext = NULL, bool bIgnoreArray = false) const;
 
 	/**
 	 * Returns true if the types are schema Equivalent. 
@@ -1102,31 +1104,31 @@ public:
 	 *
 	 * @return	true if the pin types and directions are compatible.
 	 */
-	virtual bool ArePinTypesEquivalent(const FEdGraphPinType& PinA, const FEdGraphPinType& PinB) const;
+	UE_API virtual bool ArePinTypesEquivalent(const FEdGraphPinType& PinA, const FEdGraphPinType& PinB) const;
 
 	/** Sets the autogenerated default value for a pin, optionally using the passed in function and parameter. This will also reset the current default value to the autogenerated one */
-	virtual void SetPinAutogeneratedDefaultValue(UEdGraphPin* Pin, const FString& NewValue) const;
+	UE_API virtual void SetPinAutogeneratedDefaultValue(UEdGraphPin* Pin, const FString& NewValue) const;
 
 	/** Sets the autogenerated default value for a pin using the default for that type. This will also reset the current default value to the autogenerated one */
-	virtual void SetPinAutogeneratedDefaultValueBasedOnType(UEdGraphPin* Pin) const;
+	UE_API virtual void SetPinAutogeneratedDefaultValueBasedOnType(UEdGraphPin* Pin) const;
 
 	/** Sets the pin defaults, but not autogenerated defaults, at pin construction time. This is like TrySetDefaultValue but does not do validation or callbacks */
-	virtual void SetPinDefaultValueAtConstruction(UEdGraphPin* Pin, const FString& DefaultValueString) const;
+	UE_API virtual void SetPinDefaultValueAtConstruction(UEdGraphPin* Pin, const FString& DefaultValueString) const;
 
 	/** Call to let blueprint and UI know that parameters have changed for a function/macro/etc */
-	virtual void HandleParameterDefaultValueChanged(UK2Node* TargetNode) const;
+	UE_API virtual void HandleParameterDefaultValueChanged(UK2Node* TargetNode) const;
 
 	/** Given a function and property, return the default value */
-	static bool FindFunctionParameterDefaultValue(const UFunction* Function, const FProperty* Param, FString& OutString);
+	static UE_API bool FindFunctionParameterDefaultValue(const UFunction* Function, const FProperty* Param, FString& OutString);
 
 	/** Utility that makes sure existing connections are valid, breaking any that are now illegal. */
-	static void ValidateExistingConnections(UEdGraphPin* Pin);
+	static UE_API void ValidateExistingConnections(UEdGraphPin* Pin);
 
 	/** Find a 'set value by name' function for the specified pin, if it exists */
-	static UFunction* FindSetVariableByNameFunction(const FEdGraphPinType& PinType);
+	static UE_API UFunction* FindSetVariableByNameFunction(const FEdGraphPinType& PinType);
 
 	UE_DEPRECATED(5.2, "Use the FSearchForAutocastFunctionResults variant.")
-	virtual bool SearchForAutocastFunction(const FEdGraphPinType& OutputPinType, const FEdGraphPinType& InputPinType, /*out*/ FName& TargetFunction, /*out*/ UClass*& FunctionOwner) const;
+	UE_API virtual bool SearchForAutocastFunction(const FEdGraphPinType& OutputPinType, const FEdGraphPinType& InputPinType, /*out*/ FName& TargetFunction, /*out*/ UClass*& FunctionOwner) const;
 
 	/** Find an appropriate function to call to perform an automatic cast operation */
 	struct FSearchForAutocastFunctionResults
@@ -1134,23 +1136,23 @@ public:
 		FName TargetFunction;
 		UClass* FunctionOwner = nullptr;
 	};
-	[[nodiscard]] virtual TOptional<FSearchForAutocastFunctionResults> SearchForAutocastFunction(const FEdGraphPinType& OutputPinType, const FEdGraphPinType& InputPinType) const;
+	[[nodiscard]] UE_API virtual TOptional<FSearchForAutocastFunctionResults> SearchForAutocastFunction(const FEdGraphPinType& OutputPinType, const FEdGraphPinType& InputPinType) const;
 
 	UE_DEPRECATED(5.2, "Use the FFindSpecializedConversionNodeResults variant.")
-	virtual bool FindSpecializedConversionNode(const UEdGraphPin* OutputPin, const UEdGraphPin* InputPin, bool bCreateNode, /*out*/ class UK2Node*& TargetNode) const;
+	UE_API virtual bool FindSpecializedConversionNode(const UEdGraphPin* OutputPin, const UEdGraphPin* InputPin, bool bCreateNode, /*out*/ class UK2Node*& TargetNode) const;
 	
 	UE_DEPRECATED(5.2, "Use the FFindSpecializedConversionNodeResults variant.")
-	virtual bool FindSpecializedConversionNode(const FEdGraphPinType& OutputPinType, const UEdGraphPin* InputPinType, bool bCreateNode, /*out*/ class UK2Node*& TargetNode) const;
+	UE_API virtual bool FindSpecializedConversionNode(const FEdGraphPinType& OutputPinType, const UEdGraphPin* InputPinType, bool bCreateNode, /*out*/ class UK2Node*& TargetNode) const;
 	
 	/** Find an appropriate node that can convert from one pin type to another (not a cast; e.g. "MakeLiteralArray" node) */
 	struct FFindSpecializedConversionNodeResults
 	{
 		class UK2Node* TargetNode = nullptr;
 	};
-	[[nodiscard]] virtual TOptional<FFindSpecializedConversionNodeResults> FindSpecializedConversionNode(const FEdGraphPinType& OutputPinType, const UEdGraphPin& InputPin, bool bCreateNode) const;
+	[[nodiscard]] UE_API virtual TOptional<FFindSpecializedConversionNodeResults> FindSpecializedConversionNode(const FEdGraphPinType& OutputPinType, const UEdGraphPin& InputPin, bool bCreateNode) const;
 
 	/** Create menu for variable get/set nodes which refer to a variable which does not exist. */
-	void GetNonExistentVariableMenu(FToolMenuSection& Section, const UEdGraphNode* InGraphNode, UBlueprint* OwnerBlueprint) const;
+	UE_API void GetNonExistentVariableMenu(FToolMenuSection& Section, const UEdGraphNode* InGraphNode, UBlueprint* OwnerBlueprint) const;
 
 	/**
 	 * Create menu for variable get/set nodes which allows for the replacement of variables
@@ -1160,16 +1162,16 @@ public:
 	 * @param InMenuBuilder					MenuBuilder to place the menu items into
 	 * @param bInReplaceExistingVariable	TRUE if replacing an existing variable, will keep the variable from appearing on the list
 	 */
-	void GetReplaceVariableMenu(FToolMenuSection& Section, const UEdGraphNode* InGraphNode, UBlueprint* InOwnerBlueprint, bool bInReplaceExistingVariable = false) const;
+	UE_API void GetReplaceVariableMenu(FToolMenuSection& Section, const UEdGraphNode* InGraphNode, UBlueprint* InOwnerBlueprint, bool bInReplaceExistingVariable = false) const;
 
 	// Calculates an average position between the nodes owning the two specified pins
-	static FVector2D CalculateAveragePositionBetweenNodes(UEdGraphPin* InputPin, UEdGraphPin* OutputPin);
+	static UE_API FVector2D CalculateAveragePositionBetweenNodes(UEdGraphPin* InputPin, UEdGraphPin* OutputPin);
 
 	// Tries to connect any pins with matching types and directions from the conversion node to the specified input and output pins
-	void AutowireConversionNode(UEdGraphPin* InputPin, UEdGraphPin* OutputPin, UEdGraphNode* ConversionNode) const;
+	UE_API void AutowireConversionNode(UEdGraphPin* InputPin, UEdGraphPin* OutputPin, UEdGraphNode* ConversionNode) const;
 
 	/** Calculates an estimated height for the specified node */
-	static float EstimateNodeHeight( UEdGraphNode* Node );
+	static UE_API float EstimateNodeHeight( UEdGraphNode* Node );
 
 	/** 
 	 * Checks if the graph supports impure functions
@@ -1178,14 +1180,14 @@ public:
 	 *
 	 * @return				True if the graph supports impure functions
 	 */
-	bool DoesGraphSupportImpureFunctions(const UEdGraph* InGraph) const;
+	UE_API bool DoesGraphSupportImpureFunctions(const UEdGraph* InGraph) const;
 
 	/** 
 	 * Checks if the graph is marked as thread safe
 	 * @param InGraph		Graph to check
 	 * @return			True if the graph is marked theead safe
 	 */
-	bool IsGraphMarkedThreadSafe(const UEdGraph* InGraph) const;
+	UE_API bool IsGraphMarkedThreadSafe(const UEdGraph* InGraph) const;
 	
 	/**
 	 * Checks to see if the passed in function is valid in the graph for the current class
@@ -1197,7 +1199,7 @@ public:
 	 * @param	bInCalledForEach	Call for each element in an array (a node accepts array)
 	 * @param	OutReason			Allows callers to receive a localized string containing more detail when the function is determined to be invalid (optional)
 	 */
-	bool CanFunctionBeUsedInGraph(const UClass* InClass, const UFunction* InFunction, const UEdGraph* InDestGraph, uint32 InFunctionTypes, bool bInCalledForEach, FText* OutReason = nullptr) const;
+	UE_API bool CanFunctionBeUsedInGraph(const UClass* InClass, const UFunction* InFunction, const UEdGraph* InDestGraph, uint32 InFunctionTypes, bool bInCalledForEach, FText* OutReason = nullptr) const;
 
 	/**
 	 * Makes connections into/or out of the gateway node, connect directly to the associated networks on the opposite side of the tunnel
@@ -1210,25 +1212,25 @@ public:
 	 *
 	 * @return						Returns TRUE if successful
 	 */
-	bool CollapseGatewayNode(UK2Node* InNode, UEdGraphNode* InEntryNode, UEdGraphNode* InResultNode, class FKismetCompilerContext* CompilerContext = nullptr, TSet<UEdGraphNode*>* OutExpandedNodes = nullptr) const;
+	UE_API bool CollapseGatewayNode(UK2Node* InNode, UEdGraphNode* InEntryNode, UEdGraphNode* InResultNode, class FKismetCompilerContext* CompilerContext = nullptr, TSet<UEdGraphNode*>* OutExpandedNodes = nullptr) const;
 
 	/** 
 	 * Connects all of the linked pins from PinA to all of the linked pins from PinB, removing
 	 * both PinA and PinB from being linked to anything else
 	 * Requires the nodes that own the pins to be in the same graph already (post-merging)
 	 */
-	void CombineTwoPinNetsAndRemoveOldPins(UEdGraphPin* InPinA, UEdGraphPin* InPinB) const;
+	UE_API void CombineTwoPinNetsAndRemoveOldPins(UEdGraphPin* InPinA, UEdGraphPin* InPinB) const;
 
 	/**
 	 * Make links from all data pins from InOutputNode output to InInputNode input.
 	 */
-	void LinkDataPinFromOutputToInput(UEdGraphNode* InOutputNode, UEdGraphNode* InInputNode) const;
+	UE_API void LinkDataPinFromOutputToInput(UEdGraphNode* InOutputNode, UEdGraphNode* InInputNode) const;
 
 	/** Moves all connections from the old node to the new one. Returns true and destroys OldNode on success. Fails if it cannot find a mapping from an old pin. */
-	bool ReplaceOldNodeWithNew(UEdGraphNode* OldNode, UEdGraphNode* NewNode, const TMap<FName, FName>& OldPinToNewPinMap) const;
+	UE_API bool ReplaceOldNodeWithNew(UEdGraphNode* OldNode, UEdGraphNode* NewNode, const TMap<FName, FName>& OldPinToNewPinMap) const;
 
 	/** Convert a deprecated node into a function call node, called from per-node ConvertDeprecatedNode */
-	UK2Node* ConvertDeprecatedNodeToFunctionCall(UK2Node* OldNode, UFunction* NewFunction, TMap<FName, FName>& OldPinToNewPinMap, UEdGraph* Graph) const;
+	UE_API UK2Node* ConvertDeprecatedNodeToFunctionCall(UK2Node* OldNode, UFunction* NewFunction, TMap<FName, FName>& OldPinToNewPinMap, UEdGraph* Graph) const;
 
 	/** some inherited schemas don't want anim-notify actions listed, so this is an easy way to check that */
 	virtual bool DoesSupportAnimNotifyActions() const { return true; }
@@ -1237,19 +1239,19 @@ public:
 	// NonExistent Variables: Broken Get/Set Nodes where the variable is does not exist 
 
 	/** Create the variable that the broken node refers to */
-	static void OnCreateNonExistentVariable(class UK2Node_Variable* Variable, UBlueprint* OwnerBlueprint);
+	static UE_API void OnCreateNonExistentVariable(class UK2Node_Variable* Variable, UBlueprint* OwnerBlueprint);
 
 	/** Create the local variable that the broken node refers to */
-	static void OnCreateNonExistentLocalVariable(class UK2Node_Variable* Variable, UBlueprint* OwnerBlueprint);
+	static UE_API void OnCreateNonExistentLocalVariable(class UK2Node_Variable* Variable, UBlueprint* OwnerBlueprint);
 
 	/** Replace the variable that a variable node refers to when the variable it refers to does not exist */
-	static void OnReplaceVariableForVariableNode(class UK2Node_Variable* Variable, UBlueprint* OwnerBlueprint, FName VariableName, bool bIsSelfMember);
+	static UE_API void OnReplaceVariableForVariableNode(class UK2Node_Variable* Variable, UBlueprint* OwnerBlueprint, FName VariableName, bool bIsSelfMember);
 
 	/** Create sub menu that shows all possible variables that can be used to replace the existing variable reference */
-	static void GetReplaceVariableMenu(UToolMenu* Menu, class UK2Node_Variable* Variable, UBlueprint* OwnerBlueprint, bool bReplaceExistingVariable = false);
+	static UE_API void GetReplaceVariableMenu(UToolMenu* Menu, class UK2Node_Variable* Variable, UBlueprint* OwnerBlueprint, bool bReplaceExistingVariable = false);
 
 	/** Function called when the owning module is shut down */ 
-	static void Shutdown();
+	static UE_API void Shutdown();
 private:
 
 	/**
@@ -1257,15 +1259,16 @@ private:
 	 * @param [in] Function	The function to check for out parameters
 	 * @return true if there are out parameters, else false
 	 */
-	bool DoesFunctionHaveOutParameters( const UFunction* Function ) const;
+	UE_API bool DoesFunctionHaveOutParameters( const UFunction* Function ) const;
 	
-	bool ShouldShowPanelContextMenuForIncompatibleConnectionsImpl() const override;
-
-	static const UScriptStruct* VectorStruct;
-	static const UScriptStruct* Vector3fStruct;
-	static const UScriptStruct* RotatorStruct;
-	static const UScriptStruct* TransformStruct;
-	static const UScriptStruct* LinearColorStruct;
-	static const UScriptStruct* ColorStruct;
+	UE_API bool ShouldShowPanelContextMenuForIncompatibleConnectionsImpl() const override;
+
+	static UE_API const UScriptStruct* VectorStruct;
+	static UE_API const UScriptStruct* Vector3fStruct;
+	static UE_API const UScriptStruct* RotatorStruct;
+	static UE_API const UScriptStruct* TransformStruct;
+	static UE_API const UScriptStruct* LinearColorStruct;
+	static UE_API const UScriptStruct* ColorStruct;
 };
 
+#undef UE_API
