Updated headers using UnrealCodeFixup to make sure dllstorage is on methods/staticvars instead of types. Part 4/n

diff --git a/Engine/Source/Editor/BlueprintGraph/Classes/EdGraphSchema_K2_Actions.h b/Engine/Source/Editor/BlueprintGraph/Classes/EdGraphSchema_K2_Actions.h
--- a/Engine/Source/Editor/BlueprintGraph/Classes/EdGraphSchema_K2_Actions.h
+++ b/Engine/Source/Editor/BlueprintGraph/Classes/EdGraphSchema_K2_Actions.h
@@ -35,6 +35,8 @@
 
 #include "EdGraphSchema_K2_Actions.generated.h"
 
+#define UE_API BLUEPRINTGRAPH_API
+
 class UBlueprint;
 class UEdGraph;
 class UEdGraphNode;
@@ -55,7 +57,7 @@ ENUM_CLASS_FLAGS(EK2NewNodeFlags)
 
 /** Action to add a node to the graph */
 USTRUCT()
-struct BLUEPRINTGRAPH_API FEdGraphSchemaAction_K2NewNode : public FEdGraphSchemaAction
+struct FEdGraphSchemaAction_K2NewNode : public FEdGraphSchemaAction
 {
 	GENERATED_USTRUCT_BODY()
 
@@ -84,9 +86,9 @@ struct BLUEPRINTGRAPH_API FEdGraphSchemaAction_K2NewNode : public FEdGraphSchema
 
 	// FEdGraphSchemaAction interface
 	using FEdGraphSchemaAction::PerformAction; // Prevent hiding of deprecated base class function with FVector2D
-	virtual UEdGraphNode* PerformAction(UEdGraph* ParentGraph, UEdGraphPin* FromPin, const FVector2f& Location, bool bSelectNewNode = true) override;
-	virtual UEdGraphNode* PerformAction(UEdGraph* ParentGraph, TArray<UEdGraphPin*>& FromPins, const FVector2f& Location, bool bSelectNewNode = true) override;
-	virtual void AddReferencedObjects(FReferenceCollector& Collector) override;
+	UE_API virtual UEdGraphNode* PerformAction(UEdGraph* ParentGraph, UEdGraphPin* FromPin, const FVector2f& Location, bool bSelectNewNode = true) override;
+	UE_API virtual UEdGraphNode* PerformAction(UEdGraph* ParentGraph, TArray<UEdGraphPin*>& FromPins, const FVector2f& Location, bool bSelectNewNode = true) override;
+	UE_API virtual void AddReferencedObjects(FReferenceCollector& Collector) override;
 	// End of FEdGraphSchemaAction interface
 
 	template <typename NodeType>
@@ -125,8 +127,8 @@ struct BLUEPRINTGRAPH_API FEdGraphSchemaAction_K2NewNode : public FEdGraphSchema
 		return SpawnNode<NodeType>(ParentGraph, Location, Options, [](NodeType*){});
 	}
 
-	static UEdGraphNode* CreateNode(UEdGraph* ParentGraph, TArrayView<UEdGraphPin*> FromPins, const FVector2D Location, UK2Node* NodeTemplate, EK2NewNodeFlags Options);
-	static UEdGraphNode* CreateNode(UEdGraph* ParentGraph, TArrayView<UEdGraphPin*> FromPins, const FVector2D Location, TFunctionRef<UK2Node*(UEdGraph*)> ConstructionFn, TFunctionRef<void(UK2Node*)> InitializerFn, EK2NewNodeFlags Options);
+	static UE_API UEdGraphNode* CreateNode(UEdGraph* ParentGraph, TArrayView<UEdGraphPin*> FromPins, const FVector2D Location, UK2Node* NodeTemplate, EK2NewNodeFlags Options);
+	static UE_API UEdGraphNode* CreateNode(UEdGraph* ParentGraph, TArrayView<UEdGraphPin*> FromPins, const FVector2D Location, TFunctionRef<UK2Node*(UEdGraph*)> ConstructionFn, TFunctionRef<void(UK2Node*)> InitializerFn, EK2NewNodeFlags Options);
 };
 
 /*******************************************************************************
@@ -135,7 +137,7 @@ struct BLUEPRINTGRAPH_API FEdGraphSchemaAction_K2NewNode : public FEdGraphSchema
 
 /** Action to view a node to the graph */
 USTRUCT()
-struct BLUEPRINTGRAPH_API FEdGraphSchemaAction_K2ViewNode : public FEdGraphSchemaAction
+struct FEdGraphSchemaAction_K2ViewNode : public FEdGraphSchemaAction
 {
 	GENERATED_USTRUCT_BODY()
 
@@ -158,8 +160,8 @@ struct BLUEPRINTGRAPH_API FEdGraphSchemaAction_K2ViewNode : public FEdGraphSchem
 	{}
 
 	// FEdGraphSchemaAction interface
-	virtual UEdGraphNode* PerformAction(class UEdGraph* ParentGraph, UEdGraphPin* FromPin, const FVector2f& Location, bool bSelectNewNode = true) override;
-	virtual UEdGraphNode* PerformAction(class UEdGraph* ParentGraph, TArray<UEdGraphPin*>& FromPins, const FVector2f& Location, bool bSelectNewNode = true) override;
+	UE_API virtual UEdGraphNode* PerformAction(class UEdGraph* ParentGraph, UEdGraphPin* FromPin, const FVector2f& Location, bool bSelectNewNode = true) override;
+	UE_API virtual UEdGraphNode* PerformAction(class UEdGraph* ParentGraph, TArray<UEdGraphPin*>& FromPins, const FVector2f& Location, bool bSelectNewNode = true) override;
 	// End of FEdGraphSchemaAction interface
 
 };
@@ -170,7 +172,7 @@ struct BLUEPRINTGRAPH_API FEdGraphSchemaAction_K2ViewNode : public FEdGraphSchem
 
 /** Action to add a node to the graph */
 USTRUCT()
-struct BLUEPRINTGRAPH_API FEdGraphSchemaAction_K2AssignDelegate : public FEdGraphSchemaAction_K2NewNode
+struct FEdGraphSchemaAction_K2AssignDelegate : public FEdGraphSchemaAction_K2NewNode
 {
 	GENERATED_USTRUCT_BODY()
 
@@ -186,10 +188,10 @@ struct BLUEPRINTGRAPH_API FEdGraphSchemaAction_K2AssignDelegate : public FEdGrap
 	{}
 
 	// FEdGraphSchemaAction interface
-	virtual UEdGraphNode* PerformAction(class UEdGraph* ParentGraph, UEdGraphPin* FromPin, const FVector2f& Location, bool bSelectNewNode = true) override;
+	UE_API virtual UEdGraphNode* PerformAction(class UEdGraph* ParentGraph, UEdGraphPin* FromPin, const FVector2f& Location, bool bSelectNewNode = true) override;
 	// End of FEdGraphSchemaAction interface
 
-	static UEdGraphNode* AssignDelegate(class UK2Node* NodeTemplate, class UEdGraph* ParentGraph, UEdGraphPin* FromPin, const FVector2D Location, bool bSelectNewNode);
+	static UE_API UEdGraphNode* AssignDelegate(class UK2Node* NodeTemplate, class UEdGraph* ParentGraph, UEdGraphPin* FromPin, const FVector2D Location, bool bSelectNewNode);
 };
 
 /*******************************************************************************
@@ -198,7 +200,7 @@ struct BLUEPRINTGRAPH_API FEdGraphSchemaAction_K2AssignDelegate : public FEdGrap
 
 /** Action to add a node to the graph */
 USTRUCT()
-struct BLUEPRINTGRAPH_API FEdGraphSchemaAction_EventFromFunction : public FEdGraphSchemaAction
+struct FEdGraphSchemaAction_EventFromFunction : public FEdGraphSchemaAction
 {
 	GENERATED_USTRUCT_BODY()
 
@@ -219,9 +221,9 @@ struct BLUEPRINTGRAPH_API FEdGraphSchemaAction_EventFromFunction : public FEdGra
 	{}
 
 	// FEdGraphSchemaAction interface
-	virtual UEdGraphNode* PerformAction(class UEdGraph* ParentGraph, UEdGraphPin* FromPin, const FVector2f& Location, bool bSelectNewNode = true) override;
-	virtual UEdGraphNode* PerformAction(class UEdGraph* ParentGraph, TArray<UEdGraphPin*>& FromPins, const FVector2f& Location, bool bSelectNewNode = true) override;
-	virtual void AddReferencedObjects(FReferenceCollector& Collector) override;
+	UE_API virtual UEdGraphNode* PerformAction(class UEdGraph* ParentGraph, UEdGraphPin* FromPin, const FVector2f& Location, bool bSelectNewNode = true) override;
+	UE_API virtual UEdGraphNode* PerformAction(class UEdGraph* ParentGraph, TArray<UEdGraphPin*>& FromPins, const FVector2f& Location, bool bSelectNewNode = true) override;
+	UE_API virtual void AddReferencedObjects(FReferenceCollector& Collector) override;
 	// End of FEdGraphSchemaAction interface
 };
 
@@ -390,7 +392,7 @@ struct FEdGraphSchemaAction_K2AddComment : public FEdGraphSchemaAction
 
 /** Action to target a specific node on graph */
 USTRUCT()
-struct BLUEPRINTGRAPH_API FEdGraphSchemaAction_K2TargetNode : public FEdGraphSchemaAction_K2NewNode
+struct FEdGraphSchemaAction_K2TargetNode : public FEdGraphSchemaAction_K2NewNode
 {
 	GENERATED_USTRUCT_BODY()
 
@@ -409,7 +411,7 @@ struct BLUEPRINTGRAPH_API FEdGraphSchemaAction_K2TargetNode : public FEdGraphSch
 
 	// FEdGraphSchemaAction interface
 	using FEdGraphSchemaAction::PerformAction; // Prevent hiding of deprecated base class function with FVector2D
-	virtual UEdGraphNode* PerformAction(class UEdGraph* ParentGraph, UEdGraphPin* FromPin, const FVector2f& Location, bool bSelectNewNode = true) override;
+	UE_API virtual UEdGraphNode* PerformAction(class UEdGraph* ParentGraph, UEdGraphPin* FromPin, const FVector2f& Location, bool bSelectNewNode = true) override;
 	// End of FEdGraphSchemaAction interface
 };
 
@@ -419,7 +421,7 @@ struct BLUEPRINTGRAPH_API FEdGraphSchemaAction_K2TargetNode : public FEdGraphSch
 
 /** Action to paste at this location on graph*/
 USTRUCT()
-struct BLUEPRINTGRAPH_API FEdGraphSchemaAction_K2PasteHere : public FEdGraphSchemaAction
+struct FEdGraphSchemaAction_K2PasteHere : public FEdGraphSchemaAction
 {
 	GENERATED_USTRUCT_BODY()
 
@@ -437,7 +439,7 @@ struct BLUEPRINTGRAPH_API FEdGraphSchemaAction_K2PasteHere : public FEdGraphSche
 
 	// FEdGraphSchemaAction interface
 	using FEdGraphSchemaAction::PerformAction; // Prevent hiding of deprecated base class function with FVector2D
-	virtual UEdGraphNode* PerformAction(class UEdGraph* ParentGraph, UEdGraphPin* FromPin, const FVector2f& Location, bool bSelectNewNode = true) override;
+	UE_API virtual UEdGraphNode* PerformAction(class UEdGraph* ParentGraph, UEdGraphPin* FromPin, const FVector2f& Location, bool bSelectNewNode = true) override;
 	// End of FEdGraphSchemaAction interface
 };
 
@@ -447,7 +449,7 @@ struct BLUEPRINTGRAPH_API FEdGraphSchemaAction_K2PasteHere : public FEdGraphSche
 
 /** Reference to an enumeration (only used in 'docked' palette) */
 USTRUCT()
-struct BLUEPRINTGRAPH_API FEdGraphSchemaAction_K2Enum : public FEdGraphSchemaAction
+struct FEdGraphSchemaAction_K2Enum : public FEdGraphSchemaAction
 {
 	GENERATED_USTRUCT_BODY()
 
@@ -487,7 +489,7 @@ struct BLUEPRINTGRAPH_API FEdGraphSchemaAction_K2Enum : public FEdGraphSchemaAct
 
 /** Reference to a variable (only used in 'My Blueprints' but used for member variables, local variables, delegates, etc...) */
 USTRUCT()
-struct BLUEPRINTGRAPH_API FEdGraphSchemaAction_BlueprintVariableBase : public FEdGraphSchemaAction
+struct FEdGraphSchemaAction_BlueprintVariableBase : public FEdGraphSchemaAction
 {
 	GENERATED_USTRUCT_BODY()
 
@@ -565,14 +567,14 @@ public:
 	virtual bool IsVariableUsed() { return false; }
 	
 	// FEdGraphSchemaAction interface
-	virtual void MovePersistentItemToCategory(const FText& NewCategoryName) override;
-	virtual int32 GetReorderIndexInContainer() const override;
-	virtual bool ReorderToBeforeAction(TSharedRef<FEdGraphSchemaAction> OtherAction) override;
-	virtual FEdGraphSchemaActionDefiningObject GetPersistentItemDefiningObject() const override;
+	UE_API virtual void MovePersistentItemToCategory(const FText& NewCategoryName) override;
+	UE_API virtual int32 GetReorderIndexInContainer() const override;
+	UE_API virtual bool ReorderToBeforeAction(TSharedRef<FEdGraphSchemaAction> OtherAction) override;
+	UE_API virtual FEdGraphSchemaActionDefiningObject GetPersistentItemDefiningObject() const override;
 	virtual bool IsAVariable() const { return true; }
 	// End of FEdGraphSchemaAction interface
 
-	UBlueprint* GetSourceBlueprint() const;
+	UE_API UBlueprint* GetSourceBlueprint() const;
 };
 
 
@@ -582,7 +584,7 @@ public:
 
 /** Reference to a variable (only used in 'docked' palette) */
 USTRUCT()
-struct BLUEPRINTGRAPH_API FEdGraphSchemaAction_K2Var : public FEdGraphSchemaAction_BlueprintVariableBase
+struct FEdGraphSchemaAction_K2Var : public FEdGraphSchemaAction_BlueprintVariableBase
 {
 	GENERATED_USTRUCT_BODY()
 
@@ -611,7 +613,7 @@ public:
 
 /** Reference to a local variable (only used in 'docked' palette) */
 USTRUCT()
-struct BLUEPRINTGRAPH_API FEdGraphSchemaAction_K2LocalVar : public FEdGraphSchemaAction_BlueprintVariableBase
+struct FEdGraphSchemaAction_K2LocalVar : public FEdGraphSchemaAction_BlueprintVariableBase
 {
 	GENERATED_USTRUCT_BODY()
 
@@ -619,7 +621,7 @@ public:
 	// Simple type info
 	static FName StaticGetTypeId() {static FName Type("FEdGraphSchemaAction_K2LocalVar"); return Type;}
 	virtual FName GetTypeId() const override { return StaticGetTypeId(); } 
-	virtual int32 GetReorderIndexInContainer() const override;
+	UE_API virtual int32 GetReorderIndexInContainer() const override;
 
 	FEdGraphSchemaAction_K2LocalVar() 
 		: FEdGraphSchemaAction_BlueprintVariableBase()
@@ -656,7 +658,7 @@ namespace EEdGraphSchemaAction_K2Graph
 
 /** Reference to a function, macro, event graph, or timeline (only used in 'docked' palette) */
 USTRUCT()
-struct BLUEPRINTGRAPH_API FEdGraphSchemaAction_K2Graph : public FEdGraphSchemaAction
+struct FEdGraphSchemaAction_K2Graph : public FEdGraphSchemaAction
 {
 	GENERATED_USTRUCT_BODY()
 
@@ -685,14 +687,14 @@ struct BLUEPRINTGRAPH_API FEdGraphSchemaAction_K2Graph : public FEdGraphSchemaAc
 
 	// FEdGraphSchemaAction interface
 	virtual bool IsParentable() const override { return true; }
-	virtual void MovePersistentItemToCategory(const FText& NewCategoryName) override;
-	virtual int32 GetReorderIndexInContainer() const override;
-	virtual bool ReorderToBeforeAction(TSharedRef<FEdGraphSchemaAction> OtherAction) override;
-	virtual FEdGraphSchemaActionDefiningObject GetPersistentItemDefiningObject() const override;
+	UE_API virtual void MovePersistentItemToCategory(const FText& NewCategoryName) override;
+	UE_API virtual int32 GetReorderIndexInContainer() const override;
+	UE_API virtual bool ReorderToBeforeAction(TSharedRef<FEdGraphSchemaAction> OtherAction) override;
+	UE_API virtual FEdGraphSchemaActionDefiningObject GetPersistentItemDefiningObject() const override;
 	// End of FEdGraphSchemaAction interface
 
-	UFunction* GetFunction() const;
-	UBlueprint* GetSourceBlueprint() const;
+	UE_API UFunction* GetFunction() const;
+	UE_API UBlueprint* GetSourceBlueprint() const;
 };
 
 /*******************************************************************************
@@ -704,7 +706,7 @@ struct BLUEPRINTGRAPH_API FEdGraphSchemaAction_K2Graph : public FEdGraphSchemaAc
 * to be used the 'docked' palette only.
 */
 USTRUCT()
-struct BLUEPRINTGRAPH_API FEdGraphSchemaAction_K2Event : public FEdGraphSchemaAction_K2TargetNode
+struct FEdGraphSchemaAction_K2Event : public FEdGraphSchemaAction_K2TargetNode
 {
 	GENERATED_USTRUCT_BODY()
 
@@ -754,7 +756,7 @@ struct BLUEPRINTGRAPH_API FEdGraphSchemaAction_K2Event : public FEdGraphSchemaAc
 * to be used the 'docked' palette only.
 */
 USTRUCT()
-struct BLUEPRINTGRAPH_API FEdGraphSchemaAction_K2InputAction : public FEdGraphSchemaAction_K2TargetNode
+struct FEdGraphSchemaAction_K2InputAction : public FEdGraphSchemaAction_K2TargetNode
 {
 	GENERATED_USTRUCT_BODY()
 
@@ -801,7 +803,7 @@ struct BLUEPRINTGRAPH_API FEdGraphSchemaAction_K2InputAction : public FEdGraphSc
 
 /** Reference to a delegate */
 USTRUCT()
-struct BLUEPRINTGRAPH_API FEdGraphSchemaAction_K2Delegate : public FEdGraphSchemaAction_BlueprintVariableBase
+struct FEdGraphSchemaAction_K2Delegate : public FEdGraphSchemaAction_BlueprintVariableBase
 {
 	GENERATED_USTRUCT_BODY()
 
@@ -845,3 +847,4 @@ public:
 };
 
 
+#undef UE_API
diff --git a/Engine/Source/Editor/BlueprintGraph/Classes/K2Node.h b/Engine/Source/Editor/BlueprintGraph/Classes/K2Node.h
--- a/Engine/Source/Editor/BlueprintGraph/Classes/K2Node.h
+++ b/Engine/Source/Editor/BlueprintGraph/Classes/K2Node.h
@@ -24,6 +24,8 @@
 
 #include "K2Node.generated.h"
 
+#define UE_API BLUEPRINTGRAPH_API
+
 class AActor;
 class FArchive;
 class FBlueprintActionDatabaseRegistrar;
@@ -141,16 +143,16 @@ struct FOptionalPinFromProperty
 };
 
 // Manager to build or refresh a list of optional pins
-struct BLUEPRINTGRAPH_API FOptionalPinManager
+struct FOptionalPinManager
 {
 public:
 	virtual ~FOptionalPinManager() { }
 
 	/** Should the specified property be displayed by default */
-	virtual void GetRecordDefaults(FProperty* TestProperty, FOptionalPinFromProperty& Record) const;
+	UE_API virtual void GetRecordDefaults(FProperty* TestProperty, FOptionalPinFromProperty& Record) const;
 
 	/** Can this property be managed as an optional pin (with the ability to be shown or hidden) */
-	virtual bool CanTreatPropertyAsOptional(FProperty* TestProperty) const;
+	UE_API virtual bool CanTreatPropertyAsOptional(FProperty* TestProperty) const;
 
 	/** 
 	 * Reconstructs the specified property array using the SourceStruct
@@ -158,7 +160,7 @@ public:
 	 * @param SourceStruct			The source structure to update the properties array from
 	 * @param bDefaultVisibility
 	 */
-	void RebuildPropertyList(TArray<FOptionalPinFromProperty>& Properties, UStruct* SourceStruct);
+	UE_API void RebuildPropertyList(TArray<FOptionalPinFromProperty>& Properties, UStruct* SourceStruct);
 
 	/**
 	 * Creates a pin for each visible property on the specified node
@@ -169,22 +171,22 @@ public:
 	 * @param StructBasePtr	Struct that is iterated over, to represent the current version of the data
 	 * @param DefaultsPtr	Struct that is iterated over, to represent the default version of the data, used for delta serialization
 	 */
-	void CreateVisiblePins(TArray<FOptionalPinFromProperty>& Properties, UStruct* SourceStruct, EEdGraphPinDirection Direction, class UK2Node* TargetNode, uint8* StructBasePtr = nullptr, uint8* DefaultsPtr = nullptr);
+	UE_API void CreateVisiblePins(TArray<FOptionalPinFromProperty>& Properties, UStruct* SourceStruct, EEdGraphPinDirection Direction, class UK2Node* TargetNode, uint8* StructBasePtr = nullptr, uint8* DefaultsPtr = nullptr);
 
 	// Customize automatically created pins if desired
 	virtual void CustomizePinData(UEdGraphPin* Pin, FName SourcePropertyName, int32 ArrayIndex, FProperty* Property = NULL) const {}
 
 	/** Helper function to make consistent behavior between nodes that use optional pins */
-	static void CacheShownPins(const TArray<FOptionalPinFromProperty>& OptionalPins, TArray<FName>& OldShownPins);
+	static UE_API void CacheShownPins(const TArray<FOptionalPinFromProperty>& OptionalPins, TArray<FName>& OldShownPins);
 
 	/** Helper function to make consistent behavior between nodes that use optional pins */
-	static void EvaluateOldShownPins(const TArray<FOptionalPinFromProperty>& OptionalPins, TArray<FName>& OldShownPins, UK2Node* Node);
+	static UE_API void EvaluateOldShownPins(const TArray<FOptionalPinFromProperty>& OptionalPins, TArray<FName>& OldShownPins, UK2Node* Node);
 
 
 protected:
 	virtual void PostInitNewPin(UEdGraphPin* Pin, FOptionalPinFromProperty& Record, int32 ArrayIndex, FProperty* Property, uint8* PropertyAddress, uint8* DefaultPropertyAddress) const {}
 	virtual void PostRemovedOldPin(FOptionalPinFromProperty& Record, int32 ArrayIndex, FProperty* Property, uint8* PropertyAddress, uint8* DefaultPropertyAddress) const {}
-	void RebuildProperty(FProperty* TestProperty, FName CategoryName, TArray<FOptionalPinFromProperty>& Properties, UStruct* SourceStruct, TMap<FName, FOldOptionalPinSettings>& OldSettings);
+	UE_API void RebuildProperty(FProperty* TestProperty, FName CategoryName, TArray<FOptionalPinFromProperty>& Properties, UStruct* SourceStruct, TMap<FName, FOldOptionalPinSettings>& OldSettings);
 };
 
 enum ERenamePinResult
@@ -220,7 +222,7 @@ public:
 	BLUEPRINTGRAPH_API virtual void JumpToDefinition() const override;
 	BLUEPRINTGRAPH_API virtual FString GetDocumentationLink() const override;
 	BLUEPRINTGRAPH_API virtual void GetPinHoverText(const UEdGraphPin& Pin, FString& HoverTextOut) const override;
-	BLUEPRINTGRAPH_API virtual bool ShowPaletteIconOnNode() const override { return true; }
+	virtual bool ShowPaletteIconOnNode() const override { return true; }
 	BLUEPRINTGRAPH_API virtual bool CanSplitPin(const UEdGraphPin* Pin) const override;
 	BLUEPRINTGRAPH_API virtual UEdGraphPin* GetPassThroughPin(const UEdGraphPin* FromPin) const override;
 	BLUEPRINTGRAPH_API virtual bool IsInDevelopmentMode() const override;
@@ -295,12 +297,12 @@ public:
 	BLUEPRINTGRAPH_API virtual void ExpandNode(class FKismetCompilerContext& CompilerContext, UEdGraph* SourceGraph);
 
 	/** Clears out any cached data that needs to be regenerated after a structural blueprint change */
-	BLUEPRINTGRAPH_API virtual void ClearCachedBlueprintData(UBlueprint* Blueprint) {}
+	virtual void ClearCachedBlueprintData(UBlueprint* Blueprint) {}
 
 	/** Performs a node-specific deprecation fixup, which may delete this node and replace it with another one */
-	BLUEPRINTGRAPH_API virtual void ConvertDeprecatedNode(UEdGraph* Graph, bool bOnlySafeChanges) {}
+	virtual void ConvertDeprecatedNode(UEdGraph* Graph, bool bOnlySafeChanges) {}
 
-	BLUEPRINTGRAPH_API virtual class FNodeHandlingFunctor* CreateNodeHandler(class FKismetCompilerContext& CompilerContext) const { return nullptr; }
+	virtual class FNodeHandlingFunctor* CreateNodeHandler(class FKismetCompilerContext& CompilerContext) const { return nullptr; }
 
 	BLUEPRINTGRAPH_API UK2Node* ExpandSplitPin(class FKismetCompilerContext* CompilerContext, UEdGraph* SourceGraph, UEdGraphPin* Pin);
 
@@ -343,10 +345,10 @@ public:
 	BLUEPRINTGRAPH_API virtual bool CanCreateUnderSpecifiedSchema(const UEdGraphSchema* DesiredSchema) const override;
 
 	// Returns which dynamic binding class (if any) to use for this node
-	BLUEPRINTGRAPH_API virtual UClass* GetDynamicBindingClass() const { return NULL; }
+	virtual UClass* GetDynamicBindingClass() const { return NULL; }
 
 	// Puts information about this node into the dynamic binding object
-	BLUEPRINTGRAPH_API virtual void RegisterDynamicBinding(UDynamicBlueprintBinding* BindingObject) const { };
+	virtual void RegisterDynamicBinding(UDynamicBlueprintBinding* BindingObject) const { };
 
 	// Renames a user defined pin and broadcasts the result to any OnUserDefinedPinRenamed listeners on success (if bTest is false)
 	BLUEPRINTGRAPH_API ERenamePinResult RenameUserDefinedPin(const FName OldName, const FName NewName, bool bTest = false);
@@ -367,10 +369,10 @@ public:
 	virtual bool AllowMultipleSelfs(bool bInputAsArray) const { return false; }
 
 	/** @return name of brush for special icon in upper-right corner */
-	BLUEPRINTGRAPH_API virtual FName GetCornerIcon() const { return FName(); }
+	virtual FName GetCornerIcon() const { return FName(); }
 
 	/** @return true, is pins cannot be connected due to node's inner logic, put message for user in OutReason */
-	BLUEPRINTGRAPH_API virtual bool IsConnectionDisallowed(const UEdGraphPin* MyPin, const UEdGraphPin* OtherPin, FString& OutReason) const { return false; }
+	virtual bool IsConnectionDisallowed(const UEdGraphPin* MyPin, const UEdGraphPin* OtherPin, FString& OutReason) const { return false; }
 
 	/** This function if used for nodes that needs CDO for validation (Called before expansion)*/
 	BLUEPRINTGRAPH_API virtual void EarlyValidation(FCompilerResultsLog& MessageLog) const;
@@ -379,7 +381,7 @@ public:
 	BLUEPRINTGRAPH_API virtual void GetNodeAttributes(TArray<TKeyValuePair<FString, FString>>& OutNodeAttributes) const;
 
 	/** Called before compilation begins, giving a blueprint time to force the linker to load data */
-	BLUEPRINTGRAPH_API virtual void PreloadRequiredAssets() { }
+	virtual void PreloadRequiredAssets() { }
 
 	/** 
 	 * Replacement for GetMenuEntries(). Override to add specific 
@@ -419,11 +421,11 @@ public:
 		Normal = 200,
 	};
 
-	BLUEPRINTGRAPH_API virtual int32 GetNodeRefreshPriority() const { return EBaseNodeRefreshPriority::Normal; }
+	virtual int32 GetNodeRefreshPriority() const { return EBaseNodeRefreshPriority::Normal; }
 
 	BLUEPRINTGRAPH_API bool DoesWildcardPinAcceptContainer(const UEdGraphPin* Pin);
-	BLUEPRINTGRAPH_API virtual bool DoesInputWildcardPinAcceptArray(const UEdGraphPin* Pin) const { return true; }
-	BLUEPRINTGRAPH_API virtual bool DoesOutputWildcardPinAcceptContainer(const UEdGraphPin* Pin) const { return true; }
+	virtual bool DoesInputWildcardPinAcceptArray(const UEdGraphPin* Pin) const { return true; }
+	virtual bool DoesOutputWildcardPinAcceptContainer(const UEdGraphPin* Pin) const { return true; }
 
 	/** Handle when a variable is renamed in the Blueprint Palette */
 	virtual void HandleVariableRenamed(UBlueprint* InBlueprint, UClass* InVariableClass, UEdGraph* InGraph, const FName& InOldVarName, const FName& InNewVarName) {}
@@ -593,3 +595,5 @@ public:
 	/** Util to get the generated class from a node. */
 	BLUEPRINTGRAPH_API UClass* GetBlueprintClassFromNode() const;
 };
+
+#undef UE_API
diff --git a/Engine/Source/Editor/BlueprintGraph/Classes/K2Node_AddComponentByClass.h b/Engine/Source/Editor/BlueprintGraph/Classes/K2Node_AddComponentByClass.h
--- a/Engine/Source/Editor/BlueprintGraph/Classes/K2Node_AddComponentByClass.h
+++ b/Engine/Source/Editor/BlueprintGraph/Classes/K2Node_AddComponentByClass.h
@@ -11,6 +11,8 @@
 
 #include "K2Node_AddComponentByClass.generated.h"
 
+#define UE_API BLUEPRINTGRAPH_API
+
 class FKismetCompilerContext;
 class UClass;
 class UEdGraph;
@@ -20,37 +22,39 @@ class UObject;
 /**
  * Implementation of K2Node for creating a component based on a selected or passed in class
  */
-UCLASS()
-class BLUEPRINTGRAPH_API UK2Node_AddComponentByClass : public UK2Node_ConstructObjectFromClass
+UCLASS(MinimalAPI)
+class UK2Node_AddComponentByClass : public UK2Node_ConstructObjectFromClass
 {
 	GENERATED_BODY()
 
 public:
-	UK2Node_AddComponentByClass(const FObjectInitializer& ObjectInitializer);
+	UE_API UK2Node_AddComponentByClass(const FObjectInitializer& ObjectInitializer);
 	
 	//~ Begin UEdGraphNode Interface.
-	virtual void AllocateDefaultPins() override;
-	virtual void ExpandNode(FKismetCompilerContext& CompilerContext, UEdGraph* SourceGraph) override;
-	virtual void PinDefaultValueChanged(UEdGraphPin* Pin) override;
+	UE_API virtual void AllocateDefaultPins() override;
+	UE_API virtual void ExpandNode(FKismetCompilerContext& CompilerContext, UEdGraph* SourceGraph) override;
+	UE_API virtual void PinDefaultValueChanged(UEdGraphPin* Pin) override;
 	//~ End UEdGraphNode Interface.
 
 	//~ Begin UK2Node_ConstructObjectFromClass Interface
-	virtual void CreatePinsForClass(UClass* InClass, TArray<UEdGraphPin*>* OutClassPins) override;
+	UE_API virtual void CreatePinsForClass(UClass* InClass, TArray<UEdGraphPin*>* OutClassPins) override;
 
 protected:
-	virtual FText GetBaseNodeTitle() const override;
-	virtual FText GetDefaultNodeTitle() const override;
-	virtual FText GetNodeTitleFormat() const override;
-	virtual UClass* GetClassPinBaseClass() const override;
-	virtual bool IsSpawnVarPin(UEdGraphPin* Pin) const override;
+	UE_API virtual FText GetBaseNodeTitle() const override;
+	UE_API virtual FText GetDefaultNodeTitle() const override;
+	UE_API virtual FText GetNodeTitleFormat() const override;
+	UE_API virtual UClass* GetClassPinBaseClass() const override;
+	UE_API virtual bool IsSpawnVarPin(UEdGraphPin* Pin) const override;
 	//~ End UK2Node_ConstructObjectFromClass Interface
 
-	UEdGraphPin* GetRelativeTransformPin() const;
-	UEdGraphPin* GetManualAttachmentPin() const;
+	UE_API UEdGraphPin* GetRelativeTransformPin() const;
+	UE_API UEdGraphPin* GetManualAttachmentPin() const;
 
 	/** Returns true if the currently selected or linked class is known to be a scene component */
-	bool IsSceneComponent() const;
+	UE_API bool IsSceneComponent() const;
 
 	/** Utility function to set whether the scene component specific pins are hidden or not */
-	void SetSceneComponentPinsHidden(bool bHidden);
+	UE_API void SetSceneComponentPinsHidden(bool bHidden);
 };
+
+#undef UE_API
diff --git a/Engine/Source/Editor/BlueprintGraph/Classes/K2Node_AddPinInterface.h b/Engine/Source/Editor/BlueprintGraph/Classes/K2Node_AddPinInterface.h
--- a/Engine/Source/Editor/BlueprintGraph/Classes/K2Node_AddPinInterface.h
+++ b/Engine/Source/Editor/BlueprintGraph/Classes/K2Node_AddPinInterface.h
@@ -12,10 +12,12 @@
 
 #include "K2Node_AddPinInterface.generated.h"
 
+#define UE_API BLUEPRINTGRAPH_API
+
 class UObject;
 
-UINTERFACE(meta=(CannotImplementInterfaceInBlueprint))
-class BLUEPRINTGRAPH_API UK2Node_AddPinInterface : public UInterface
+UINTERFACE(MinimalAPI, meta=(CannotImplementInterfaceInBlueprint))
+class UK2Node_AddPinInterface : public UInterface
 {
 	GENERATED_BODY()
 };
@@ -29,7 +31,7 @@ class UEdGraphPin;
 * 
 * @see UK2Node_CommutativeAssociativeBinaryOperator, UK2Node_DoOnceMultiInput
 */
-class BLUEPRINTGRAPH_API IK2Node_AddPinInterface
+class IK2Node_AddPinInterface
 {
 	GENERATED_BODY()
 
@@ -54,7 +56,7 @@ public:
 	*
 	* @return FName			Name of this additional pin
 	*/
-	static FName GetNameForAdditionalPin(int32 PinIndex);
+	static UE_API FName GetNameForAdditionalPin(int32 PinIndex);
 
 	/**
 	* Add an additional input pin to this node
@@ -86,3 +88,5 @@ public:
 	*/
 	virtual bool CanRemovePin(const UEdGraphPin* Pin) const { return true; };
 };
+
+#undef UE_API
diff --git a/Engine/Source/Editor/BlueprintGraph/Classes/K2Node_AsyncAction.h b/Engine/Source/Editor/BlueprintGraph/Classes/K2Node_AsyncAction.h
--- a/Engine/Source/Editor/BlueprintGraph/Classes/K2Node_AsyncAction.h
+++ b/Engine/Source/Editor/BlueprintGraph/Classes/K2Node_AsyncAction.h
@@ -9,20 +9,24 @@
 
 #include "K2Node_AsyncAction.generated.h"
 
+#define UE_API BLUEPRINTGRAPH_API
+
 class FBlueprintActionDatabaseRegistrar;
 class UObject;
 
 /** !!! The proxy object should have RF_StrongRefOnFrame flag. !!! */
 
-UCLASS()
-class BLUEPRINTGRAPH_API UK2Node_AsyncAction : public UK2Node_BaseAsyncTask
+UCLASS(MinimalAPI)
+class UK2Node_AsyncAction : public UK2Node_BaseAsyncTask
 {
 	GENERATED_UCLASS_BODY()
 	
 	// UK2Node interface
-	virtual void GetMenuActions(FBlueprintActionDatabaseRegistrar& ActionRegistrar) const override;
+	UE_API virtual void GetMenuActions(FBlueprintActionDatabaseRegistrar& ActionRegistrar) const override;
 	// End of UK2Node interface
 
 	/** Initialize the async task from a known function **/
-	void InitializeProxyFromFunction(const UFunction* ProxyFunction);
+	UE_API void InitializeProxyFromFunction(const UFunction* ProxyFunction);
 };
+
+#undef UE_API
diff --git a/Engine/Source/Editor/BlueprintGraph/Classes/K2Node_BaseAsyncTask.h b/Engine/Source/Editor/BlueprintGraph/Classes/K2Node_BaseAsyncTask.h
--- a/Engine/Source/Editor/BlueprintGraph/Classes/K2Node_BaseAsyncTask.h
+++ b/Engine/Source/Editor/BlueprintGraph/Classes/K2Node_BaseAsyncTask.h
@@ -19,6 +19,8 @@
 
 #include "K2Node_BaseAsyncTask.generated.h"
 
+#define UE_API BLUEPRINTGRAPH_API
+
 class FBlueprintActionDatabaseRegistrar;
 class FKismetCompilerContext;
 class FMulticastDelegateProperty;
@@ -48,36 +50,36 @@ struct FAsyncTaskPinRedirectMapInfo
  * Base class used for a blueprint node that calls a function and provides asynchronous exec pins when the function is complete.
  * The created proxy object should have RF_StrongRefOnFrame flag.
  */
-UCLASS(Abstract)
-class BLUEPRINTGRAPH_API UK2Node_BaseAsyncTask : public UK2Node
+UCLASS(MinimalAPI, Abstract)
+class UK2Node_BaseAsyncTask : public UK2Node
 {
 	GENERATED_UCLASS_BODY()
 
 public:
 	// UEdGraphNode interface
-	virtual void AllocateDefaultPins() override;
-	virtual FText GetTooltipText() const override;
-	virtual FText GetNodeTitle(ENodeTitleType::Type TitleType) const override;
-	virtual bool IsCompatibleWithGraph(const UEdGraph* TargetGraph) const override;
-	virtual void ValidateNodeDuringCompilation(class FCompilerResultsLog& MessageLog) const override;
-	virtual void GetPinHoverText(const UEdGraphPin& Pin, FString& HoverTextOut) const override;
-	virtual FString GetPinMetaData(FName InPinName, FName InKey) override;
-	virtual bool HasExternalDependencies(TArray<class UStruct*>* OptionalOutput) const override;
-	virtual bool IsLatentForMacros() const override;
+	UE_API virtual void AllocateDefaultPins() override;
+	UE_API virtual FText GetTooltipText() const override;
+	UE_API virtual FText GetNodeTitle(ENodeTitleType::Type TitleType) const override;
+	UE_API virtual bool IsCompatibleWithGraph(const UEdGraph* TargetGraph) const override;
+	UE_API virtual void ValidateNodeDuringCompilation(class FCompilerResultsLog& MessageLog) const override;
+	UE_API virtual void GetPinHoverText(const UEdGraphPin& Pin, FString& HoverTextOut) const override;
+	UE_API virtual FString GetPinMetaData(FName InPinName, FName InKey) override;
+	UE_API virtual bool HasExternalDependencies(TArray<class UStruct*>* OptionalOutput) const override;
+	UE_API virtual bool IsLatentForMacros() const override;
 	// End of UEdGraphNode interface
 
 	// UK2Node interface
-	virtual void ExpandNode(FKismetCompilerContext& CompilerContext, UEdGraph* SourceGraph) override;
-	virtual FName GetCornerIcon() const override;
-	virtual FText GetToolTipHeading() const override;
-	virtual void GetMenuActions(FBlueprintActionDatabaseRegistrar& ActionRegistrar) const override;
-	virtual FText GetMenuCategory() const override;
-	virtual void GetRedirectPinNames(const UEdGraphPin& Pin, TArray<FString>& RedirectPinNames) const override;
-	virtual ERedirectType DoPinsMatchForReconstruction(const UEdGraphPin* NewPin, int32 NewPinIndex, const UEdGraphPin* OldPin, int32 OldPinIndex) const override;
+	UE_API virtual void ExpandNode(FKismetCompilerContext& CompilerContext, UEdGraph* SourceGraph) override;
+	UE_API virtual FName GetCornerIcon() const override;
+	UE_API virtual FText GetToolTipHeading() const override;
+	UE_API virtual void GetMenuActions(FBlueprintActionDatabaseRegistrar& ActionRegistrar) const override;
+	UE_API virtual FText GetMenuCategory() const override;
+	UE_API virtual void GetRedirectPinNames(const UEdGraphPin& Pin, TArray<FString>& RedirectPinNames) const override;
+	UE_API virtual ERedirectType DoPinsMatchForReconstruction(const UEdGraphPin* NewPin, int32 NewPinIndex, const UEdGraphPin* OldPin, int32 OldPinIndex) const override;
 	// End of UK2Node interface
 
 	/** Returns the primary function that this node will call */
-	UFunction* GetFactoryFunction() const;
+	UE_API UFunction* GetFactoryFunction() const;
 
 protected:
 	/**
@@ -91,7 +93,7 @@ protected:
 	 * @return	True if a successful connection was made to an intermediate node
 	 *			or if one was not necessary. False if a connection was failed.
 	 */
-	bool ExpandDefaultToSelfPin(FKismetCompilerContext& CompilerContext, UEdGraph* SourceGraph, UK2Node_CallFunction* IntermediateProxyNode);
+	UE_API bool ExpandDefaultToSelfPin(FKismetCompilerContext& CompilerContext, UEdGraph* SourceGraph, UK2Node_CallFunction* IntermediateProxyNode);
 	
 	// The name of the function to call to create a proxy object
 	UPROPERTY()
@@ -109,7 +111,7 @@ protected:
 	UPROPERTY()
 	FName ProxyActivateFunctionName;
 
-	struct BLUEPRINTGRAPH_API FBaseAsyncTaskHelper
+	struct FBaseAsyncTaskHelper
 	{
 		struct FOutputPinAndLocalVariable
 		{
@@ -119,23 +121,23 @@ protected:
 			FOutputPinAndLocalVariable(UEdGraphPin* Pin, UK2Node_TemporaryVariable* Var) : OutputPin(Pin), TempVar(Var) {}
 		};
 
-		static bool ValidDataPin(const UEdGraphPin* Pin, EEdGraphPinDirection Direction);
-		static bool CreateDelegateForNewFunction(UEdGraphPin* DelegateInputPin, FName FunctionName, UK2Node* CurrentNode, UEdGraph* SourceGraph, FKismetCompilerContext& CompilerContext);
-		static bool CopyEventSignature(UK2Node_CustomEvent* CENode, UFunction* Function, const UEdGraphSchema_K2* Schema);
-		static bool HandleDelegateImplementation(
+		static UE_API bool ValidDataPin(const UEdGraphPin* Pin, EEdGraphPinDirection Direction);
+		static UE_API bool CreateDelegateForNewFunction(UEdGraphPin* DelegateInputPin, FName FunctionName, UK2Node* CurrentNode, UEdGraph* SourceGraph, FKismetCompilerContext& CompilerContext);
+		static UE_API bool CopyEventSignature(UK2Node_CustomEvent* CENode, UFunction* Function, const UEdGraphSchema_K2* Schema);
+		static UE_API bool HandleDelegateImplementation(
 			FMulticastDelegateProperty* CurrentProperty, const TArray<FBaseAsyncTaskHelper::FOutputPinAndLocalVariable>& VariableOutputs,
 			UEdGraphPin* ProxyObjectPin, UEdGraphPin*& InOutLastThenPin, UEdGraphPin*& OutLastActivatedThenPin,
 			UK2Node* CurrentNode, UEdGraph* SourceGraph, FKismetCompilerContext& CompilerContext);
 
-		static const FName GetAsyncTaskProxyName();
+		static UE_API const FName GetAsyncTaskProxyName();
 	};
 
 	// Pin Redirector support
-	static TMap<FName, FAsyncTaskPinRedirectMapInfo> AsyncTaskPinRedirectMap;
-	static bool bAsyncTaskPinRedirectMapInitialized;
+	static UE_API TMap<FName, FAsyncTaskPinRedirectMapInfo> AsyncTaskPinRedirectMap;
+	static UE_API bool bAsyncTaskPinRedirectMapInitialized;
 
 	/** Expand out the logic to handle the delegate output pins */
-	virtual bool HandleDelegates(
+	UE_API virtual bool HandleDelegates(
 		const TArray<FBaseAsyncTaskHelper::FOutputPinAndLocalVariable>& VariableOutputs, UEdGraphPin* ProxyObjectPin,
 		UEdGraphPin*& InOutLastThenPin, UEdGraph* SourceGraph, FKismetCompilerContext& CompilerContext);
 
@@ -148,8 +150,10 @@ private:
 	*
 	* @param   Pin				The pin you want hover text for (should belong to this node)
 	*/
-	void GeneratePinTooltip(UEdGraphPin& Pin) const;
+	UE_API void GeneratePinTooltip(UEdGraphPin& Pin) const;
 
 	/** Flag used to track validity of pin tooltips, when tooltips are invalid they will be refreshed before being displayed */
 	mutable bool bPinTooltipsValid;
 };
+
+#undef UE_API
diff --git a/Engine/Source/Editor/BlueprintGraph/Classes/K2Node_BaseMCDelegate.h b/Engine/Source/Editor/BlueprintGraph/Classes/K2Node_BaseMCDelegate.h
--- a/Engine/Source/Editor/BlueprintGraph/Classes/K2Node_BaseMCDelegate.h
+++ b/Engine/Source/Editor/BlueprintGraph/Classes/K2Node_BaseMCDelegate.h
@@ -43,22 +43,22 @@ public:
 	virtual bool HasExternalDependencies(TArray<class UStruct*>* OptionalOutput) const override;
 	// End of UEdGraphNode interface
 
-	BLUEPRINTGRAPH_API void SetFromProperty(const FProperty* Property, bool bSelfContext, UClass* OwnerClass)
+	void SetFromProperty(const FProperty* Property, bool bSelfContext, UClass* OwnerClass)
 	{
 		DelegateReference.SetFromField<FProperty>(Property, bSelfContext, OwnerClass);
 	}
 
-	BLUEPRINTGRAPH_API FProperty* GetProperty() const
+	FProperty* GetProperty() const
 	{
 		return DelegateReference.ResolveMember<FMulticastDelegateProperty>(GetBlueprintClassFromNode());
 	}
 
-	BLUEPRINTGRAPH_API FName GetPropertyName() const
+	FName GetPropertyName() const
 	{
 		return DelegateReference.GetMemberName();
 	}
 
-	BLUEPRINTGRAPH_API FText GetPropertyDisplayName() const
+	FText GetPropertyDisplayName() const
 	{
 		FProperty* Prop = GetProperty();
 		return (Prop ? Prop->GetDisplayNameText() : FText::FromName(GetPropertyName()));
diff --git a/Engine/Source/Editor/BlueprintGraph/Classes/K2Node_CallFunction.h b/Engine/Source/Editor/BlueprintGraph/Classes/K2Node_CallFunction.h
--- a/Engine/Source/Editor/BlueprintGraph/Classes/K2Node_CallFunction.h
+++ b/Engine/Source/Editor/BlueprintGraph/Classes/K2Node_CallFunction.h
@@ -24,6 +24,8 @@
 
 #include "K2Node_CallFunction.generated.h"
 
+#define UE_API BLUEPRINTGRAPH_API
+
 class FArchive;
 class FKismetCompilerContext;
 class FProperty;
@@ -42,8 +44,8 @@ enum class ENodePurityOverride : int8
 	Impure
 };
 
-UCLASS()
-class BLUEPRINTGRAPH_API UK2Node_CallFunction : public UK2Node
+UCLASS(MinimalAPI)
+class UK2Node_CallFunction : public UK2Node
 {
 	GENERATED_UCLASS_BODY()
 
@@ -96,73 +98,73 @@ private:
 public:
 
 	// UObject interface
-	virtual void PostDuplicate(bool bDuplicateForPIE) override;
-	virtual void Serialize(FArchive& Ar) override;
+	UE_API virtual void PostDuplicate(bool bDuplicateForPIE) override;
+	UE_API virtual void Serialize(FArchive& Ar) override;
 	// End of UObject interface
 
 	// UEdGraphNode interface
-	virtual void GetPinHoverText(const UEdGraphPin& Pin, FString& HoverTextOut) const override;
-	virtual void AllocateDefaultPins() override;
-	virtual FLinearColor GetNodeTitleColor() const override;
-	virtual FString GetFindReferenceSearchString_Impl(EGetFindReferenceSearchStringFlags InFlags) const override;
-	virtual FText GetTooltipText() const override;
-	virtual FText GetNodeTitle(ENodeTitleType::Type TitleType) const override;
-	virtual FString GetDescriptiveCompiledName() const override;
-	virtual bool HasDeprecatedReference() const override;
-	virtual FEdGraphNodeDeprecationResponse GetDeprecationResponse(EEdGraphNodeDeprecationType DeprecationType) const override;
-	virtual void PostPlacedNewNode() override;
-	virtual FString GetDocumentationLink() const override;
-	virtual FString GetDocumentationExcerptName() const override;
-	virtual FSlateIcon GetIconAndTint(FLinearColor& OutColor) const override;
-	virtual bool CanPasteHere(const UEdGraph* TargetGraph) const override;
-	virtual void PinDefaultValueChanged(UEdGraphPin* Pin) override;
-	virtual void AddSearchMetaDataInfo(TArray<struct FSearchTagDataPair>& OutTaggedMetaData) const override;
-	virtual void AddPinSearchMetaDataInfo(const UEdGraphPin* Pin, TArray<struct FSearchTagDataPair>& OutTaggedMetaData) const override;
-	virtual TSharedPtr<SWidget> CreateNodeImage() const override;
-	virtual UObject* GetJumpTargetForDoubleClick() const override;
-	virtual bool CanJumpToDefinition() const override;
-	virtual void JumpToDefinition() const override;
-	virtual void GetNodeContextMenuActions(class UToolMenu* Menu, class UGraphNodeContextMenuContext* Context) const override;
-	virtual FString GetPinMetaData(FName InPinName, FName InKey) override;
-	virtual bool HasExternalDependencies(TArray<class UStruct*>* OptionalOutput) const override;
-	virtual bool IsLatentForMacros() const override;
+	UE_API virtual void GetPinHoverText(const UEdGraphPin& Pin, FString& HoverTextOut) const override;
+	UE_API virtual void AllocateDefaultPins() override;
+	UE_API virtual FLinearColor GetNodeTitleColor() const override;
+	UE_API virtual FString GetFindReferenceSearchString_Impl(EGetFindReferenceSearchStringFlags InFlags) const override;
+	UE_API virtual FText GetTooltipText() const override;
+	UE_API virtual FText GetNodeTitle(ENodeTitleType::Type TitleType) const override;
+	UE_API virtual FString GetDescriptiveCompiledName() const override;
+	UE_API virtual bool HasDeprecatedReference() const override;
+	UE_API virtual FEdGraphNodeDeprecationResponse GetDeprecationResponse(EEdGraphNodeDeprecationType DeprecationType) const override;
+	UE_API virtual void PostPlacedNewNode() override;
+	UE_API virtual FString GetDocumentationLink() const override;
+	UE_API virtual FString GetDocumentationExcerptName() const override;
+	UE_API virtual FSlateIcon GetIconAndTint(FLinearColor& OutColor) const override;
+	UE_API virtual bool CanPasteHere(const UEdGraph* TargetGraph) const override;
+	UE_API virtual void PinDefaultValueChanged(UEdGraphPin* Pin) override;
+	UE_API virtual void AddSearchMetaDataInfo(TArray<struct FSearchTagDataPair>& OutTaggedMetaData) const override;
+	UE_API virtual void AddPinSearchMetaDataInfo(const UEdGraphPin* Pin, TArray<struct FSearchTagDataPair>& OutTaggedMetaData) const override;
+	UE_API virtual TSharedPtr<SWidget> CreateNodeImage() const override;
+	UE_API virtual UObject* GetJumpTargetForDoubleClick() const override;
+	UE_API virtual bool CanJumpToDefinition() const override;
+	UE_API virtual void JumpToDefinition() const override;
+	UE_API virtual void GetNodeContextMenuActions(class UToolMenu* Menu, class UGraphNodeContextMenuContext* Context) const override;
+	UE_API virtual FString GetPinMetaData(FName InPinName, FName InKey) override;
+	UE_API virtual bool HasExternalDependencies(TArray<class UStruct*>* OptionalOutput) const override;
+	UE_API virtual bool IsLatentForMacros() const override;
 	// End of UEdGraphNode interface
 
 	// UK2Node interface
-	virtual void ReallocatePinsDuringReconstruction(TArray<UEdGraphPin*>& OldPins) override;
-	virtual bool IsNodePure() const override;
-	virtual void PostReconstructNode() override;
-	virtual bool ShouldDrawCompact() const override;
+	UE_API virtual void ReallocatePinsDuringReconstruction(TArray<UEdGraphPin*>& OldPins) override;
+	UE_API virtual bool IsNodePure() const override;
+	UE_API virtual void PostReconstructNode() override;
+	UE_API virtual bool ShouldDrawCompact() const override;
 	UE_DEPRECATED(5.4, "ShouldDrawAsBead is deprecated")
 	virtual bool ShouldDrawAsBead() const override { return false; }
-	virtual FText GetCompactNodeTitle() const override;
-	virtual void PostPasteNode() override;
-	virtual bool CanSplitPin(const UEdGraphPin* Pin) const override;
-	virtual void ValidateNodeDuringCompilation(class FCompilerResultsLog& MessageLog) const override;
-	virtual bool ShouldShowNodeProperties() const override;
-	virtual void GetRedirectPinNames(const UEdGraphPin& Pin, TArray<FString>& RedirectPinNames) const override;
-	virtual void NotifyPinConnectionListChanged(UEdGraphPin* Pin) override;
-	virtual class FNodeHandlingFunctor* CreateNodeHandler(class FKismetCompilerContext& CompilerContext) const override;
-	virtual void ExpandNode(class FKismetCompilerContext& CompilerContext, UEdGraph* SourceGraph) override;
-	virtual FName GetCornerIcon() const override;
-	virtual FText GetToolTipHeading() const override;
-	virtual void GetNodeAttributes( TArray<TKeyValuePair<FString, FString>>& OutNodeAttributes ) const override;
-	virtual FText GetMenuCategory() const override;
-	virtual bool IsActionFilteredOut(class FBlueprintActionFilter const& Filter) override;
-	virtual bool IsConnectionDisallowed(const UEdGraphPin* MyPin, const UEdGraphPin* OtherPin, FString& OutReason) const override;
+	UE_API virtual FText GetCompactNodeTitle() const override;
+	UE_API virtual void PostPasteNode() override;
+	UE_API virtual bool CanSplitPin(const UEdGraphPin* Pin) const override;
+	UE_API virtual void ValidateNodeDuringCompilation(class FCompilerResultsLog& MessageLog) const override;
+	UE_API virtual bool ShouldShowNodeProperties() const override;
+	UE_API virtual void GetRedirectPinNames(const UEdGraphPin& Pin, TArray<FString>& RedirectPinNames) const override;
+	UE_API virtual void NotifyPinConnectionListChanged(UEdGraphPin* Pin) override;
+	UE_API virtual class FNodeHandlingFunctor* CreateNodeHandler(class FKismetCompilerContext& CompilerContext) const override;
+	UE_API virtual void ExpandNode(class FKismetCompilerContext& CompilerContext, UEdGraph* SourceGraph) override;
+	UE_API virtual FName GetCornerIcon() const override;
+	UE_API virtual FText GetToolTipHeading() const override;
+	UE_API virtual void GetNodeAttributes( TArray<TKeyValuePair<FString, FString>>& OutNodeAttributes ) const override;
+	UE_API virtual FText GetMenuCategory() const override;
+	UE_API virtual bool IsActionFilteredOut(class FBlueprintActionFilter const& Filter) override;
+	UE_API virtual bool IsConnectionDisallowed(const UEdGraphPin* MyPin, const UEdGraphPin* OtherPin, FString& OutReason) const override;
 	// End of UK2Node interface
 
 	/** Returns the UFunction that this class is pointing to */
-	UFunction* GetTargetFunction() const;
+	UE_API UFunction* GetTargetFunction() const;
 
 	/** Get the return value pin */
-	UEdGraphPin* GetReturnValuePin() const;
+	UE_API UEdGraphPin* GetReturnValuePin() const;
 
 	/** @return	true if the function is a latent operation */
-	bool IsLatentFunction() const;
+	UE_API bool IsLatentFunction() const;
 
 	/** @return true if this function can be called on multiple contexts at once */
-	virtual bool AllowMultipleSelfs(bool bInputAsArray) const override;
+	UE_API virtual bool AllowMultipleSelfs(bool bInputAsArray) const override;
 
 	/**
 	 * Creates a self pin for the graph, taking into account the scope of the function call
@@ -171,7 +173,7 @@ public:
 	 *
 	 * @return	Pointer to the pin that was created
 	 */
-	virtual UEdGraphPin* CreateSelfPin(const UFunction* Function);
+	UE_API virtual UEdGraphPin* CreateSelfPin(const UFunction* Function);
 
 	/**
 	 * Creates all of the pins required to call a particular UFunction.
@@ -180,115 +182,116 @@ public:
 	 *
 	 * @return	true on success.
 	 */
-	bool CreatePinsForFunctionCall(const UFunction* Function);
+	UE_API bool CreatePinsForFunctionCall(const UFunction* Function);
 
 	/** Create exec pins for this function. May be multiple is using 'expand enum as execs' */
-	void CreateExecPinsForFunctionCall(const UFunction* Function);
+	UE_API void CreateExecPinsForFunctionCall(const UFunction* Function);
 
 	/** Gets the name of the referenced function */
-	FName GetFunctionName() const;
+	UE_API FName GetFunctionName() const;
 
 	virtual void PostParameterPinCreated(UEdGraphPin *Pin) {}
 
 	UE_DEPRECATED(5.5, "Moved to ObjectTools::GetUserFacingFunctionName.")
-	static FText GetUserFacingFunctionName(const UFunction* Function, ENodeTitleType::Type NodeTitleType = ENodeTitleType::EditableTitle);
+	static UE_API FText GetUserFacingFunctionName(const UFunction* Function, ENodeTitleType::Type NodeTitleType = ENodeTitleType::EditableTitle);
 
 	/** Set up a pins tooltip from a function's tooltip */
-	static void GeneratePinTooltipFromFunction(UEdGraphPin& Pin, const UFunction* Function);
+	static UE_API void GeneratePinTooltipFromFunction(UEdGraphPin& Pin, const UFunction* Function);
 
 	UE_DEPRECATED(5.5, "Moved to ObjectTools::GetDefaultTooltipForFunction.")
-	static FString GetDefaultTooltipForFunction(const UFunction* Function);
+	static UE_API FString GetDefaultTooltipForFunction(const UFunction* Function);
 
 	/** Get default category for this function in action menu */
-	static FText GetDefaultCategoryForFunction(const UFunction* Function, const FText& BaseCategory);
+	static UE_API FText GetDefaultCategoryForFunction(const UFunction* Function, const FText& BaseCategory);
 	/** Get keywords for this function in the action menu */
-	static FText GetKeywordsForFunction(const UFunction* Function);
+	static UE_API FText GetKeywordsForFunction(const UFunction* Function);
 	/** Should be drawn compact for this function */
-	static bool ShouldDrawCompact(const UFunction* Function);
+	static UE_API bool ShouldDrawCompact(const UFunction* Function);
 	/** Get the compact name for this function */
-	static FString GetCompactNodeTitle(const UFunction* Function);
+	static UE_API FString GetCompactNodeTitle(const UFunction* Function);
 
 	/** Get the text to use to explain the context for this function (used on node title) */
-	virtual FText GetFunctionContextString() const;
+	UE_API virtual FText GetFunctionContextString() const;
 
 	/** Set properties of this node from a supplied function (does not save ref to function) */
-	virtual void SetFromFunction(const UFunction* Function);
+	UE_API virtual void SetFromFunction(const UFunction* Function);
 
-	static void CallForEachElementInArrayExpansion(UK2Node* Node, UEdGraphPin* MultiSelf, FKismetCompilerContext& CompilerContext, UEdGraph* SourceGraph);
+	static UE_API void CallForEachElementInArrayExpansion(UK2Node* Node, UEdGraphPin* MultiSelf, FKismetCompilerContext& CompilerContext, UEdGraph* SourceGraph);
 
-	static UEdGraphPin* InnerHandleAutoCreateRef(UK2Node* Node, UEdGraphPin* Pin, FKismetCompilerContext& CompilerContext, UEdGraph* SourceGraph, bool bForceAssignment);
+	static UE_API UEdGraphPin* InnerHandleAutoCreateRef(UK2Node* Node, UEdGraphPin* Pin, FKismetCompilerContext& CompilerContext, UEdGraph* SourceGraph, bool bForceAssignment);
 
 	/**
 	 * Returns the graph for this function, if available. In cases of calling an event, it will be the ubergraph for the event
 	 *
 	 * @param OutGraphNode		If this function calls an event, this param is the event node, otherwise it is NULL
 	 */
-	UEdGraph* GetFunctionGraph(const UEdGraphNode*& OutGraphNode) const;
+	UE_API UEdGraph* GetFunctionGraph(const UEdGraphNode*& OutGraphNode) const;
 
 	/** Checks if the property is marked as "CustomStructureParam" */
-	static bool IsStructureWildcardProperty(const UFunction* InFunction, const FName PropertyName);
+	static UE_API bool IsStructureWildcardProperty(const UFunction* InFunction, const FName PropertyName);
 
 	/** returns true if InProperty should be treated as a wildcard (e.g. due to SetParam markup) */
-	static bool IsWildcardProperty(const UFunction* InFunction, const FProperty* InProperty);
+	static UE_API bool IsWildcardProperty(const UFunction* InFunction, const FProperty* InProperty);
 
 	/** Used to determine the result of AllowMultipleSelfs() (without having a node instance) */
-	static bool CanFunctionSupportMultipleTargets(UFunction const* InFunction);
+	static UE_API bool CanFunctionSupportMultipleTargets(UFunction const* InFunction);
 
 	/** Checks if the input function can be called in the input object with respect to editor-only/runtime mismatch*/
-	static bool CanEditorOnlyFunctionBeCalled(const UFunction* InFunction, const UObject* InObject);
+	static UE_API bool CanEditorOnlyFunctionBeCalled(const UFunction* InFunction, const UObject* InObject);
 
 	/** */
-	static FSlateIcon GetPaletteIconForFunction(UFunction const* Function, FLinearColor& OutColor);
+	static UE_API FSlateIcon GetPaletteIconForFunction(UFunction const* Function, FLinearColor& OutColor);
 
-	static void GetExpandEnumPinNames(const UFunction* Function, TArray<FName>& EnumNamesToCheck);
+	static UE_API void GetExpandEnumPinNames(const UFunction* Function, TArray<FName>& EnumNamesToCheck);
 
 private: 
 	/* Looks at function metadata and properties to determine if this node should be using enum to exec expansion */
-	void DetermineWantsEnumToExecExpansion(const UFunction* Function);
+	UE_API void DetermineWantsEnumToExecExpansion(const UFunction* Function);
 
 	/**
 	 * Creates hover text for the specified pin.
 	 * 
 	 * @param   Pin				The pin you want hover text for (should belong to this node)
 	 */
-	void GeneratePinTooltip(UEdGraphPin& Pin) const;
+	UE_API void GeneratePinTooltip(UEdGraphPin& Pin) const;
 
 	/**
 	 * Connect Execute and Then pins for functions, which became pure.
 	 */
-	bool ReconnectPureExecPins(TArray<UEdGraphPin*>& OldPins);
+	UE_API bool ReconnectPureExecPins(TArray<UEdGraphPin*>& OldPins);
 
 	/** Conforms container pins */
-	void ConformContainerPins();
+	UE_API void ConformContainerPins();
 
 	UPROPERTY()
 	ENodePurityOverride NodePurityOverride;
 
-	bool AreExecPinsVisible() const;
-	bool FunctionHasOutputs() const;
-	void ToggleNodePurityOverride();
+	UE_API bool AreExecPinsVisible() const;
+	UE_API bool FunctionHasOutputs() const;
+	UE_API void ToggleNodePurityOverride();
 
 protected:
 
 	/** Invalidates current pin tool tips, so that they will be refreshed before being displayed: */
-	void InvalidatePinTooltips();
+	UE_API void InvalidatePinTooltips();
 
 	/** Helper function to ensure function is called in our context */
-	virtual void FixupSelfMemberContext();
+	UE_API virtual void FixupSelfMemberContext();
 
 	/** By default, pure nodes can be toggled. Return false if you don't want your node to support toggling. */
-	virtual bool CanToggleNodePurity() const;
+	UE_API virtual bool CanToggleNodePurity() const;
 
 	/** Adds this function to the suppressed deprecation warnings list for this project */
-	void SuppressDeprecationWarning() const;
+	UE_API void SuppressDeprecationWarning() const;
 
 	/** Helper function for searching a UFunction for the names of requires pins/params: */
-	static TSet<FName> GetRequiredParamNames(const UFunction* ForFunction);
+	static UE_API TSet<FName> GetRequiredParamNames(const UFunction* ForFunction);
 
 	/** Routine for validating that all UPARAM(Required) parmas have a connection: */
-	void ValidateRequiredPins(const UFunction* Function, class FCompilerResultsLog& MessageLog) const;
+	UE_API void ValidateRequiredPins(const UFunction* Function, class FCompilerResultsLog& MessageLog) const;
 
 	/** Helper function to find UFunction entries from the skeleton class, use with caution.. */
-	UFunction* GetTargetFunctionFromSkeletonClass() const;
+	UE_API UFunction* GetTargetFunctionFromSkeletonClass() const;
 };
 
+#undef UE_API
diff --git a/Engine/Source/Editor/BlueprintGraph/Classes/K2Node_CommutativeAssociativeBinaryOperator.h b/Engine/Source/Editor/BlueprintGraph/Classes/K2Node_CommutativeAssociativeBinaryOperator.h
--- a/Engine/Source/Editor/BlueprintGraph/Classes/K2Node_CommutativeAssociativeBinaryOperator.h
+++ b/Engine/Source/Editor/BlueprintGraph/Classes/K2Node_CommutativeAssociativeBinaryOperator.h
@@ -42,7 +42,7 @@ public:
 	BLUEPRINTGRAPH_API UEdGraphPin* GetInputPin(int32 InputPinIndex);
 
 	/** Returns the number of additional input pins that this node has */
-	BLUEPRINTGRAPH_API int32 GetNumberOfAdditionalInputs() const { return NumAdditionalInputs; }
+	int32 GetNumberOfAdditionalInputs() const { return NumAdditionalInputs; }
 
 	virtual void RemoveInputPin(UEdGraphPin* Pin) override;
 
diff --git a/Engine/Source/Editor/BlueprintGraph/Classes/K2Node_ComponentBoundEvent.h b/Engine/Source/Editor/BlueprintGraph/Classes/K2Node_ComponentBoundEvent.h
--- a/Engine/Source/Editor/BlueprintGraph/Classes/K2Node_ComponentBoundEvent.h
+++ b/Engine/Source/Editor/BlueprintGraph/Classes/K2Node_ComponentBoundEvent.h
@@ -79,7 +79,7 @@ class UK2Node_ComponentBoundEvent : public UK2Node_Event
 
 public:
 
-	BLUEPRINTGRAPH_API FName GetComponentPropertyName() const { return ComponentPropertyName; }
+	FName GetComponentPropertyName() const { return ComponentPropertyName; }
 
 private:
 
diff --git a/Engine/Source/Editor/BlueprintGraph/Classes/K2Node_Composite.h b/Engine/Source/Editor/BlueprintGraph/Classes/K2Node_Composite.h
--- a/Engine/Source/Editor/BlueprintGraph/Classes/K2Node_Composite.h
+++ b/Engine/Source/Editor/BlueprintGraph/Classes/K2Node_Composite.h
@@ -17,11 +17,13 @@
 
 #include "K2Node_Composite.generated.h"
 
+#define UE_API BLUEPRINTGRAPH_API
+
 class INameValidatorInterface;
 class UObject;
 
-UCLASS()
-class BLUEPRINTGRAPH_API UK2Node_Composite : public UK2Node_Tunnel
+UCLASS(MinimalAPI)
+class UK2Node_Composite : public UK2Node_Tunnel
 {
 	GENERATED_UCLASS_BODY()
 
@@ -30,23 +32,23 @@ class BLUEPRINTGRAPH_API UK2Node_Composite : public UK2Node_Tunnel
 	TObjectPtr<class UEdGraph> BoundGraph;
 
 	//~ Begin UObject Interface
-	virtual void PostEditUndo() override;
+	UE_API virtual void PostEditUndo() override;
 	//~ End UObject Interface
 
 	//~ Begin UEdGraphNode Interface
-	virtual void AllocateDefaultPins() override;
-	virtual void DestroyNode() override;
-	virtual void PostPasteNode() override;
-	virtual FText GetTooltipText() const override;
-	virtual FLinearColor GetNodeTitleColor() const override;
-	virtual FText GetNodeTitle(ENodeTitleType::Type TitleType) const override;
-	virtual bool CanUserDeleteNode() const override;
-	virtual UObject* GetJumpTargetForDoubleClick() const override;
-	virtual void PostPlacedNewNode() override;
-	virtual void OnRenameNode(const FString& NewName) override;
-	virtual TSharedPtr<class INameValidatorInterface> MakeNameValidator() const override;
+	UE_API virtual void AllocateDefaultPins() override;
+	UE_API virtual void DestroyNode() override;
+	UE_API virtual void PostPasteNode() override;
+	UE_API virtual FText GetTooltipText() const override;
+	UE_API virtual FLinearColor GetNodeTitleColor() const override;
+	UE_API virtual FText GetNodeTitle(ENodeTitleType::Type TitleType) const override;
+	UE_API virtual bool CanUserDeleteNode() const override;
+	UE_API virtual UObject* GetJumpTargetForDoubleClick() const override;
+	UE_API virtual void PostPlacedNewNode() override;
+	UE_API virtual void OnRenameNode(const FString& NewName) override;
+	UE_API virtual TSharedPtr<class INameValidatorInterface> MakeNameValidator() const override;
 	virtual TArray<UEdGraph*> GetSubGraphs() const override { return TArray<UEdGraph*>( { BoundGraph } ); }
-	virtual bool CanCreateUnderSpecifiedSchema(const UEdGraphSchema* DesiredSchema) const override;
+	UE_API virtual bool CanCreateUnderSpecifiedSchema(const UEdGraphSchema* DesiredSchema) const override;
 	//~ End UEdGraphNode Interface
 
 	//~ Begin UK2Node Interface
@@ -56,22 +58,22 @@ class BLUEPRINTGRAPH_API UK2Node_Composite : public UK2Node_Tunnel
 	//~ End UK2Node Interface
 
 	// Get the entry/exit nodes inside this collapsed graph
-	UK2Node_Tunnel* GetEntryNode() const;
-	UK2Node_Tunnel* GetExitNode() const;
+	UE_API UK2Node_Tunnel* GetEntryNode() const;
+	UE_API UK2Node_Tunnel* GetExitNode() const;
 
 protected:
 	/** Fixes up the input and output sink when needed, useful after PostEditUndo which changes which graph these nodes point to */
-	void FixupInputAndOutputSink();
+	UE_API void FixupInputAndOutputSink();
 
 private:
 	/** Rename the BoundGraph to a unique name
 		- this is a special case rename: RenameGraphCloseToName() assumes that we are only concurrned with the immediate Outer() scope,
 		  but in the case of the BoundGraph we must also look into the Outer of the composite node(a Graph), and make sure no graphs in its SubGraph
 		  array are already using the name.  */
-	void RenameBoundGraphCloseToName(const FString& Name);
+	UE_API void RenameBoundGraphCloseToName(const FString& Name);
 
 	/** Determine if the name already used by another graph in composite nodes chain */
-	bool IsCompositeNameAvailable( const FString& NewName );
+	UE_API bool IsCompositeNameAvailable( const FString& NewName );
 
 	/** Constructing FText strings can be costly, so we cache the node's title */
 	FNodeTextCache CachedNodeTitle;
@@ -79,3 +81,4 @@ private:
 
 
 
+#undef UE_API
diff --git a/Engine/Source/Editor/BlueprintGraph/Classes/K2Node_ConstructObjectFromClass.h b/Engine/Source/Editor/BlueprintGraph/Classes/K2Node_ConstructObjectFromClass.h
--- a/Engine/Source/Editor/BlueprintGraph/Classes/K2Node_ConstructObjectFromClass.h
+++ b/Engine/Source/Editor/BlueprintGraph/Classes/K2Node_ConstructObjectFromClass.h
@@ -16,6 +16,8 @@
 
 #include "K2Node_ConstructObjectFromClass.generated.h"
 
+#define UE_API BLUEPRINTGRAPH_API
+
 class FBlueprintActionDatabaseRegistrar;
 class UClass;
 class UEdGraph;
@@ -23,68 +25,68 @@ class UEdGraphPin;
 class UObject;
 template <typename KeyType, typename ValueType> struct TKeyValuePair;
 
-UCLASS(abstract)
-class BLUEPRINTGRAPH_API UK2Node_ConstructObjectFromClass : public UK2Node
+UCLASS(MinimalAPI, abstract)
+class UK2Node_ConstructObjectFromClass : public UK2Node
 {
 	GENERATED_UCLASS_BODY()
 
 	//~ Begin UEdGraphNode Interface.
-	virtual void AllocateDefaultPins() override;
-	virtual FText GetNodeTitle(ENodeTitleType::Type TitleType) const override;
-	virtual void PinDefaultValueChanged(UEdGraphPin* Pin) override;
-	virtual FText GetTooltipText() const override;
-	virtual FText GetKeywords() const override;
-	virtual bool HasExternalDependencies(TArray<class UStruct*>* OptionalOutput) const override;
-	virtual bool IsCompatibleWithGraph(const UEdGraph* TargetGraph) const override;
-	virtual void PinConnectionListChanged(UEdGraphPin* Pin);
-	virtual void GetPinHoverText(const UEdGraphPin& Pin, FString& HoverTextOut) const override;
-	virtual void PostPlacedNewNode() override;
-	virtual void AddSearchMetaDataInfo(TArray<struct FSearchTagDataPair>& OutTaggedMetaData) const override;
-	virtual FString GetPinMetaData(FName InPinName, FName InKey) override;
+	UE_API virtual void AllocateDefaultPins() override;
+	UE_API virtual FText GetNodeTitle(ENodeTitleType::Type TitleType) const override;
+	UE_API virtual void PinDefaultValueChanged(UEdGraphPin* Pin) override;
+	UE_API virtual FText GetTooltipText() const override;
+	UE_API virtual FText GetKeywords() const override;
+	UE_API virtual bool HasExternalDependencies(TArray<class UStruct*>* OptionalOutput) const override;
+	UE_API virtual bool IsCompatibleWithGraph(const UEdGraph* TargetGraph) const override;
+	UE_API virtual void PinConnectionListChanged(UEdGraphPin* Pin);
+	UE_API virtual void GetPinHoverText(const UEdGraphPin& Pin, FString& HoverTextOut) const override;
+	UE_API virtual void PostPlacedNewNode() override;
+	UE_API virtual void AddSearchMetaDataInfo(TArray<struct FSearchTagDataPair>& OutTaggedMetaData) const override;
+	UE_API virtual FString GetPinMetaData(FName InPinName, FName InKey) override;
 	//~ End UEdGraphNode Interface.
 
 	//~ Begin UK2Node Interface
 	virtual bool IsNodeSafeToIgnore() const override { return true; }
-	virtual void ReallocatePinsDuringReconstruction(TArray<UEdGraphPin*>& OldPins) override;
-	virtual ERedirectType DoPinsMatchForReconstruction(const UEdGraphPin* NewPin, int32 NewPinIndex, const UEdGraphPin* OldPin, int32 OldPinIndex) const override;
-	virtual void GetNodeAttributes( TArray<TKeyValuePair<FString, FString>>& OutNodeAttributes ) const override;
-	virtual void GetMenuActions(FBlueprintActionDatabaseRegistrar& ActionRegistrar) const override;
-	virtual FText GetMenuCategory() const override;
+	UE_API virtual void ReallocatePinsDuringReconstruction(TArray<UEdGraphPin*>& OldPins) override;
+	UE_API virtual ERedirectType DoPinsMatchForReconstruction(const UEdGraphPin* NewPin, int32 NewPinIndex, const UEdGraphPin* OldPin, int32 OldPinIndex) const override;
+	UE_API virtual void GetNodeAttributes( TArray<TKeyValuePair<FString, FString>>& OutNodeAttributes ) const override;
+	UE_API virtual void GetMenuActions(FBlueprintActionDatabaseRegistrar& ActionRegistrar) const override;
+	UE_API virtual FText GetMenuCategory() const override;
 	//~ End UK2Node Interface
 
 	/** Create new pins to show properties on archetype */
-	virtual void CreatePinsForClass(UClass* InClass, TArray<UEdGraphPin*>* OutClassPins = nullptr);
+	UE_API virtual void CreatePinsForClass(UClass* InClass, TArray<UEdGraphPin*>* OutClassPins = nullptr);
 
 	/** See if this is a spawn variable pin, or a 'default' pin */
-	virtual bool IsSpawnVarPin(UEdGraphPin* Pin) const;
+	UE_API virtual bool IsSpawnVarPin(UEdGraphPin* Pin) const;
 
 	/** Get the blueprint input pin */	
-	UEdGraphPin* GetClassPin(const TArray<UEdGraphPin*>* InPinsToSearch=NULL) const;
+	UE_API UEdGraphPin* GetClassPin(const TArray<UEdGraphPin*>* InPinsToSearch=NULL) const;
 	/** Get the world context input pin, can return NULL */
-	UEdGraphPin* GetWorldContextPin() const;
+	UE_API UEdGraphPin* GetWorldContextPin() const;
 	/** Get the result output pin */
-	UEdGraphPin* GetResultPin() const;
+	UE_API UEdGraphPin* GetResultPin() const;
 	/** Get the result input pin */
-	UEdGraphPin* GetOuterPin() const;
+	UE_API UEdGraphPin* GetOuterPin() const;
 
 	/** Get the class that we are going to spawn, if it's defined as default value */
-	UClass* GetClassToSpawn(const TArray<UEdGraphPin*>* InPinsToSearch=NULL) const;
+	UE_API UClass* GetClassToSpawn(const TArray<UEdGraphPin*>* InPinsToSearch=NULL) const;
 
 	/** Returns if the node uses World Object Context input */
-	virtual bool UseWorldContext() const;
+	UE_API virtual bool UseWorldContext() const;
 
 	/** Returns if the node uses Outer input */
 	virtual bool UseOuter() const { return false; }
 
 protected:
 	/** Gets the node for use in lists and menus */
-	virtual FText GetBaseNodeTitle() const;
+	UE_API virtual FText GetBaseNodeTitle() const;
 	/** Gets the default node title when no class is selected */
-	virtual FText GetDefaultNodeTitle() const;
+	UE_API virtual FText GetDefaultNodeTitle() const;
 	/** Gets the node title when a class has been selected. */
-	virtual FText GetNodeTitleFormat() const;
+	UE_API virtual FText GetNodeTitleFormat() const;
 	/** Gets base class to use for the 'class' pin.  UObject by default. */
-	virtual UClass* GetClassPinBaseClass() const;
+	UE_API virtual UClass* GetClassPinBaseClass() const;
 
 	/**
 	 * Takes the specified "MutatablePin" and sets its 'PinToolTip' field (according
@@ -93,10 +95,10 @@ protected:
 	 * @param   MutatablePin	The pin you want to set tool-tip text on
 	 * @param   PinDescription	A string describing the pin's purpose
 	 */
-	void SetPinToolTip(UEdGraphPin& MutatablePin, const FText& PinDescription) const;
+	UE_API void SetPinToolTip(UEdGraphPin& MutatablePin, const FText& PinDescription) const;
 
 	/** Refresh pins when class was changed */
-	void OnClassPinChanged();
+	UE_API void OnClassPinChanged();
 
 	/** Tooltip text for this node. */
 	FText NodeTooltip;
@@ -104,3 +106,5 @@ protected:
 	/** Constructing FText strings can be costly, so we cache the node's title */
 	FNodeTextCache CachedNodeTitle;
 };
+
+#undef UE_API
diff --git a/Engine/Source/Editor/BlueprintGraph/Classes/K2Node_DelegateSet.h b/Engine/Source/Editor/BlueprintGraph/Classes/K2Node_DelegateSet.h
--- a/Engine/Source/Editor/BlueprintGraph/Classes/K2Node_DelegateSet.h
+++ b/Engine/Source/Editor/BlueprintGraph/Classes/K2Node_DelegateSet.h
@@ -58,7 +58,7 @@ class UK2Node_DelegateSet : public UK2Node
 	BLUEPRINTGRAPH_API UEdGraphPin* GetDelegateOwner() const;
 
 	// @todo document
-	BLUEPRINTGRAPH_API FName GetDelegateTargetEntryPointName() const
+	FName GetDelegateTargetEntryPointName() const
 	{
 		const FString TargetName = GetName() + TEXT("_") + DelegatePropertyName.ToString() + TEXT("_EP");
 		return FName(*TargetName);
diff --git a/Engine/Source/Editor/BlueprintGraph/Classes/K2Node_EditablePinBase.h b/Engine/Source/Editor/BlueprintGraph/Classes/K2Node_EditablePinBase.h
--- a/Engine/Source/Editor/BlueprintGraph/Classes/K2Node_EditablePinBase.h
+++ b/Engine/Source/Editor/BlueprintGraph/Classes/K2Node_EditablePinBase.h
@@ -141,7 +141,7 @@ class UK2Node_EditablePinBase : public UK2Node
 	/** Pins defined by the user */
 	TArray< TSharedPtr<FUserPinInfo> >UserDefinedPins;
 
-	BLUEPRINTGRAPH_API virtual bool IsEditable() const { return bIsEditable; }
+	virtual bool IsEditable() const { return bIsEditable; }
 
 	// UObject interface
 	BLUEPRINTGRAPH_API virtual void Serialize(FArchive& Ar) override;
@@ -156,7 +156,7 @@ class UK2Node_EditablePinBase : public UK2Node
 	// End of UEdGraphNode interface
 
 	// UK2Node interface
-	BLUEPRINTGRAPH_API virtual bool ShouldShowNodeProperties() const override { return bIsEditable; }
+	virtual bool ShouldShowNodeProperties() const override { return bIsEditable; }
 	// End of UK2Node interface
 
 	/**
@@ -166,7 +166,7 @@ class UK2Node_EditablePinBase : public UK2Node
 	 * @param			OutErrorMessage		Only filled with an error if there is pin add support but there is an error with the pin type
 	 * @return			TRUE if a user defined pin can be constructed
 	 */
-	BLUEPRINTGRAPH_API virtual bool CanCreateUserDefinedPin(const FEdGraphPinType& InPinType, EEdGraphPinDirection InDesiredDirection, FText& OutErrorMessage) { return false; }
+	virtual bool CanCreateUserDefinedPin(const FEdGraphPinType& InPinType, EEdGraphPinDirection InDesiredDirection, FText& OutErrorMessage) { return false; }
 
 	/**
 	 * Creates a UserPinInfo from the specified information, and also adds a pin based on that description to the node
diff --git a/Engine/Source/Editor/BlueprintGraph/Classes/K2Node_Event.h b/Engine/Source/Editor/BlueprintGraph/Classes/K2Node_Event.h
--- a/Engine/Source/Editor/BlueprintGraph/Classes/K2Node_Event.h
+++ b/Engine/Source/Editor/BlueprintGraph/Classes/K2Node_Event.h
@@ -121,7 +121,7 @@ class UK2Node_Event : public UK2Node_EditablePinBase, public IK2Node_EventNodeIn
 	BLUEPRINTGRAPH_API UFunction* FindEventSignatureFunction() const;
 	BLUEPRINTGRAPH_API void UpdateDelegatePin(bool bSilent = false);
 	BLUEPRINTGRAPH_API FName GetFunctionName() const;
-	BLUEPRINTGRAPH_API virtual bool IsUsedByAuthorityOnlyDelegate() const { return false; }
+	virtual bool IsUsedByAuthorityOnlyDelegate() const { return false; }
 	BLUEPRINTGRAPH_API virtual bool IsCosmeticTickEvent() const;
 
 	/** Returns localized string describing replication settings. 
diff --git a/Engine/Source/Editor/BlueprintGraph/Classes/K2Node_EventNodeInterface.h b/Engine/Source/Editor/BlueprintGraph/Classes/K2Node_EventNodeInterface.h
--- a/Engine/Source/Editor/BlueprintGraph/Classes/K2Node_EventNodeInterface.h
+++ b/Engine/Source/Editor/BlueprintGraph/Classes/K2Node_EventNodeInterface.h
@@ -9,13 +9,13 @@
 
 struct FEdGraphSchemaAction;
 
-UINTERFACE(meta=(CannotImplementInterfaceInBlueprint))
-class BLUEPRINTGRAPH_API UK2Node_EventNodeInterface : public UInterface
+UINTERFACE(MinimalAPI, meta=(CannotImplementInterfaceInBlueprint))
+class UK2Node_EventNodeInterface : public UInterface
 {
 	GENERATED_BODY()
 };
 
-class BLUEPRINTGRAPH_API IK2Node_EventNodeInterface
+class IK2Node_EventNodeInterface
 {
 	GENERATED_BODY()
 
diff --git a/Engine/Source/Editor/BlueprintGraph/Classes/K2Node_ExternalGraphInterface.h b/Engine/Source/Editor/BlueprintGraph/Classes/K2Node_ExternalGraphInterface.h
--- a/Engine/Source/Editor/BlueprintGraph/Classes/K2Node_ExternalGraphInterface.h
+++ b/Engine/Source/Editor/BlueprintGraph/Classes/K2Node_ExternalGraphInterface.h
@@ -8,17 +8,17 @@
 
 class UEdGraph;
 
-UINTERFACE(meta=(CannotImplementInterfaceInBlueprint))
-class BLUEPRINTGRAPH_API UK2Node_ExternalGraphInterface : public UInterface
+UINTERFACE(MinimalAPI, meta=(CannotImplementInterfaceInBlueprint))
+class UK2Node_ExternalGraphInterface : public UInterface
 {
 	GENERATED_BODY()
 };
 
-class BLUEPRINTGRAPH_API IK2Node_ExternalGraphInterface
+class IK2Node_ExternalGraphInterface
 {
 	GENERATED_BODY()
 
 public:
 	// Get external graphs to display for this node
 	virtual TArray<UEdGraph*> GetExternalGraphs() const = 0;
-};
\ No newline at end of file
+};
diff --git a/Engine/Source/Editor/BlueprintGraph/Classes/K2Node_FunctionEntry.h b/Engine/Source/Editor/BlueprintGraph/Classes/K2Node_FunctionEntry.h
--- a/Engine/Source/Editor/BlueprintGraph/Classes/K2Node_FunctionEntry.h
+++ b/Engine/Source/Editor/BlueprintGraph/Classes/K2Node_FunctionEntry.h
@@ -116,35 +116,35 @@ class UK2Node_FunctionEntry : public UK2Node_FunctionTerminator
 	BLUEPRINTGRAPH_API int32 GetFunctionFlags() const;
 
 	/** Retrieves the extra flags set on this node. */
-	BLUEPRINTGRAPH_API int32 GetExtraFlags() const
+	int32 GetExtraFlags() const
 	{
 		return ExtraFlags;
 	}
 
 	/** Set the extra flags on this node */
-	BLUEPRINTGRAPH_API void SetExtraFlags(int32 InFlags)
+	void SetExtraFlags(int32 InFlags)
 	{
 		ExtraFlags = (InFlags & ~FUNC_Native);
 	}
 
-	BLUEPRINTGRAPH_API void AddExtraFlags(int32 InFlags)
+	void AddExtraFlags(int32 InFlags)
 	{
 		ExtraFlags |= InFlags;
 	}
 
-	BLUEPRINTGRAPH_API void ClearExtraFlags(int32 InFlags)
+	void ClearExtraFlags(int32 InFlags)
 	{
 		ExtraFlags &= ~InFlags;
 	}
 	
 	/** Used to safely check whether the passed in flag is set. */
-	BLUEPRINTGRAPH_API bool HasAnyExtraFlags(int32 FlagsToCheck) const
+	bool HasAnyExtraFlags(int32 FlagsToCheck) const
 	{
 		return (ExtraFlags & FlagsToCheck) != 0 || FlagsToCheck == ~0;
 	}
 
 	/** Used to safely check whether all of the passed in flags are set. */
-	BLUEPRINTGRAPH_API bool HasAllExtraFlags(int32 FlagsToCheck) const
+	bool HasAllExtraFlags(int32 FlagsToCheck) const
 	{
 		return ((ExtraFlags & FlagsToCheck) == FlagsToCheck);
 	}
diff --git a/Engine/Source/Editor/BlueprintGraph/Classes/K2Node_GenericCreateObject.h b/Engine/Source/Editor/BlueprintGraph/Classes/K2Node_GenericCreateObject.h
--- a/Engine/Source/Editor/BlueprintGraph/Classes/K2Node_GenericCreateObject.h
+++ b/Engine/Source/Editor/BlueprintGraph/Classes/K2Node_GenericCreateObject.h
@@ -9,20 +9,22 @@
 
 #include "K2Node_GenericCreateObject.generated.h"
 
+#define UE_API BLUEPRINTGRAPH_API
+
 class FKismetCompilerContext;
 class UEdGraph;
 class UK2Node_CallFunction;
 class UObject;
 
-UCLASS()
-class BLUEPRINTGRAPH_API UK2Node_GenericCreateObject : public UK2Node_ConstructObjectFromClass
+UCLASS(MinimalAPI)
+class UK2Node_GenericCreateObject : public UK2Node_ConstructObjectFromClass
 {
 	GENERATED_BODY()
 
 	//~ Begin UEdGraphNode Interface.
-	virtual void ExpandNode(class FKismetCompilerContext& CompilerContext, UEdGraph* SourceGraph) override;
-	virtual void EarlyValidation(class FCompilerResultsLog& MessageLog) const override;
-	virtual bool IsCompatibleWithGraph(const UEdGraph* TargetGraph) const override;
+	UE_API virtual void ExpandNode(class FKismetCompilerContext& CompilerContext, UEdGraph* SourceGraph) override;
+	UE_API virtual void EarlyValidation(class FCompilerResultsLog& MessageLog) const override;
+	UE_API virtual bool IsCompatibleWithGraph(const UEdGraph* TargetGraph) const override;
 	//~ End UEdGraphNode Interface.
 
 	virtual bool UseWorldContext() const override { return false; }
@@ -37,5 +39,7 @@ class BLUEPRINTGRAPH_API UK2Node_GenericCreateObject : public UK2Node_ConstructO
 	 *
 	 * @return	true on success.
 	 */
-	bool ExpandDefaultToSelfPin(FKismetCompilerContext& CompilerContext, UEdGraph* SourceGraph, UK2Node_CallFunction* CallCreateNode);
+	UE_API bool ExpandDefaultToSelfPin(FKismetCompilerContext& CompilerContext, UEdGraph* SourceGraph, UK2Node_CallFunction* CallCreateNode);
 };
+
+#undef UE_API
diff --git a/Engine/Source/Editor/BlueprintGraph/Classes/K2Node_GetClassDefaults.h b/Engine/Source/Editor/BlueprintGraph/Classes/K2Node_GetClassDefaults.h
--- a/Engine/Source/Editor/BlueprintGraph/Classes/K2Node_GetClassDefaults.h
+++ b/Engine/Source/Editor/BlueprintGraph/Classes/K2Node_GetClassDefaults.h
@@ -57,13 +57,13 @@ class UK2Node_GetClassDefaults : public UK2Node
 
 public:
 	/** Finds and returns the class input pin from the current set of pins. */
-	BLUEPRINTGRAPH_API UEdGraphPin* FindClassPin() const
+	UEdGraphPin* FindClassPin() const
 	{
 		return FindClassPin(Pins);
 	}
 
 	/** Retrieves the current input class type. */
-	BLUEPRINTGRAPH_API UClass* GetInputClass() const
+	UClass* GetInputClass() const
 	{
 		return GetInputClass(FindClassPin());
 	}
diff --git a/Engine/Source/Editor/BlueprintGraph/Classes/K2Node_GetDataTableRow.h b/Engine/Source/Editor/BlueprintGraph/Classes/K2Node_GetDataTableRow.h
--- a/Engine/Source/Editor/BlueprintGraph/Classes/K2Node_GetDataTableRow.h
+++ b/Engine/Source/Editor/BlueprintGraph/Classes/K2Node_GetDataTableRow.h
@@ -16,6 +16,8 @@
 
 #include "K2Node_GetDataTableRow.generated.h"
 
+#define UE_API BLUEPRINTGRAPH_API
+
 class FBlueprintActionDatabaseRegistrar;
 class FString;
 class UDataTable;
@@ -25,48 +27,48 @@ class UObject;
 class UScriptStruct;
 struct FLinearColor;
 
-UCLASS()
-class BLUEPRINTGRAPH_API UK2Node_GetDataTableRow : public UK2Node
+UCLASS(MinimalAPI)
+class UK2Node_GetDataTableRow : public UK2Node
 {
 	GENERATED_UCLASS_BODY()
 
 	//~ Begin UEdGraphNode Interface.
-	virtual void AllocateDefaultPins() override;
-	virtual FText GetNodeTitle(ENodeTitleType::Type TitleType) const override;
-	virtual void PinDefaultValueChanged(UEdGraphPin* Pin) override;
-	virtual FText GetTooltipText() const override;
-	virtual void ExpandNode(class FKismetCompilerContext& CompilerContext, UEdGraph* SourceGraph) override;
-	virtual FSlateIcon GetIconAndTint(FLinearColor& OutColor) const override;
-	virtual void PostReconstructNode() override;
+	UE_API virtual void AllocateDefaultPins() override;
+	UE_API virtual FText GetNodeTitle(ENodeTitleType::Type TitleType) const override;
+	UE_API virtual void PinDefaultValueChanged(UEdGraphPin* Pin) override;
+	UE_API virtual FText GetTooltipText() const override;
+	UE_API virtual void ExpandNode(class FKismetCompilerContext& CompilerContext, UEdGraph* SourceGraph) override;
+	UE_API virtual FSlateIcon GetIconAndTint(FLinearColor& OutColor) const override;
+	UE_API virtual void PostReconstructNode() override;
 	//~ End UEdGraphNode Interface.
 
 	//~ Begin UK2Node Interface
 	virtual bool IsNodeSafeToIgnore() const override { return true; }
-	virtual void ReallocatePinsDuringReconstruction(TArray<UEdGraphPin*>& OldPins) override;
-	virtual void GetMenuActions(FBlueprintActionDatabaseRegistrar& ActionRegistrar) const override;
-	virtual FText GetMenuCategory() const override;
-	virtual bool IsConnectionDisallowed(const UEdGraphPin* MyPin, const UEdGraphPin* OtherPin, FString& OutReason) const override;
-	virtual void EarlyValidation(class FCompilerResultsLog& MessageLog) const override;
-	virtual void PreloadRequiredAssets() override;
-	virtual void NotifyPinConnectionListChanged(UEdGraphPin* Pin) override;
+	UE_API virtual void ReallocatePinsDuringReconstruction(TArray<UEdGraphPin*>& OldPins) override;
+	UE_API virtual void GetMenuActions(FBlueprintActionDatabaseRegistrar& ActionRegistrar) const override;
+	UE_API virtual FText GetMenuCategory() const override;
+	UE_API virtual bool IsConnectionDisallowed(const UEdGraphPin* MyPin, const UEdGraphPin* OtherPin, FString& OutReason) const override;
+	UE_API virtual void EarlyValidation(class FCompilerResultsLog& MessageLog) const override;
+	UE_API virtual void PreloadRequiredAssets() override;
+	UE_API virtual void NotifyPinConnectionListChanged(UEdGraphPin* Pin) override;
 	//~ End UK2Node Interface
 
 	/** Get the return type of our struct */
-	UScriptStruct* GetReturnTypeForStruct();
+	UE_API UScriptStruct* GetReturnTypeForStruct();
 
 	/** Get the Data Table input pin */
-	UEdGraphPin* GetDataTablePin(const TArray<UEdGraphPin*>* InPinsToSearch=NULL) const;
+	UE_API UEdGraphPin* GetDataTablePin(const TArray<UEdGraphPin*>* InPinsToSearch=NULL) const;
 	/** Get the spawn transform input pin */	
-	UEdGraphPin* GetRowNamePin() const;
+	UE_API UEdGraphPin* GetRowNamePin() const;
     /** Get the exec output pin for when the row was not found */
-	UEdGraphPin* GetRowNotFoundPin() const;
+	UE_API UEdGraphPin* GetRowNotFoundPin() const;
 	/** Get the result output pin */
-	UEdGraphPin* GetResultPin() const;
+	UE_API UEdGraphPin* GetResultPin() const;
 
 	/** Get the type of the TableRow to return */
-	UScriptStruct* GetDataTableRowStructType() const;
+	UE_API UScriptStruct* GetDataTableRowStructType() const;
 
-	void OnDataTableRowListChanged(const UDataTable* DataTable);
+	UE_API void OnDataTableRowListChanged(const UDataTable* DataTable);
 private:
 	/**
 	 * Takes the specified "MutatablePin" and sets its 'PinToolTip' field (according
@@ -75,14 +77,14 @@ private:
 	 * @param   MutatablePin	The pin you want to set tool-tip text on
 	 * @param   PinDescription	A string describing the pin's purpose
 	 */
-	void SetPinToolTip(UEdGraphPin& MutatablePin, const FText& PinDescription) const;
+	UE_API void SetPinToolTip(UEdGraphPin& MutatablePin, const FText& PinDescription) const;
 
 	/** Set the return type of our struct */
-	void SetReturnTypeForStruct(UScriptStruct* InClass);
+	UE_API void SetReturnTypeForStruct(UScriptStruct* InClass);
 	/** Queries for the authoritative return type, then modifies the return pin to match */
-	void RefreshOutputPinType();
+	UE_API void RefreshOutputPinType();
 	/** Triggers a refresh which will update the node's widget; aimed at updating the dropdown menu for the RowName input */
-	void RefreshRowNameOptions();
+	UE_API void RefreshRowNameOptions();
 
 	/** Tooltip text for this node. */
 	FText NodeTooltip;
@@ -90,3 +92,5 @@ private:
 	/** Constructing FText strings can be costly, so we cache the node's title */
 	FNodeTextCache CachedNodeTitle;
 };
+
+#undef UE_API
diff --git a/Engine/Source/Editor/BlueprintGraph/Classes/K2Node_GetSubsystem.h b/Engine/Source/Editor/BlueprintGraph/Classes/K2Node_GetSubsystem.h
--- a/Engine/Source/Editor/BlueprintGraph/Classes/K2Node_GetSubsystem.h
+++ b/Engine/Source/Editor/BlueprintGraph/Classes/K2Node_GetSubsystem.h
@@ -16,6 +16,8 @@
 
 #include "K2Node_GetSubsystem.generated.h"
 
+#define UE_API BLUEPRINTGRAPH_API
+
 class FArchive;
 class FBlueprintActionDatabaseRegistrar;
 class FKismetCompilerContext;
@@ -27,41 +29,41 @@ class UObject;
 class USubsystem;
 template <typename KeyType, typename ValueType> struct TKeyValuePair;
 
-UCLASS()
-class BLUEPRINTGRAPH_API UK2Node_GetSubsystem : public UK2Node
+UCLASS(MinimalAPI)
+class UK2Node_GetSubsystem : public UK2Node
 {
 	GENERATED_BODY()
 public:
 
-	virtual void Serialize(FArchive& Ar) override;
+	UE_API virtual void Serialize(FArchive& Ar) override;
 
-	void Initialize(UClass* NodeClass);
+	UE_API void Initialize(UClass* NodeClass);
 	//~Begin UEdGraphNode interface.
-	virtual void AllocateDefaultPins() override;
-	virtual bool IsCompatibleWithGraph(const UEdGraph* TargetGraph) const override;
-	virtual FSlateIcon GetIconAndTint(FLinearColor& OutColor) const override;
-	virtual FLinearColor GetNodeTitleColor() const override;
-	virtual FText GetNodeTitle(ENodeTitleType::Type TitleType) const override;
-	virtual void ExpandNode(FKismetCompilerContext& CompilerContext, UEdGraph* SourceGraph) override;
-	virtual bool CanJumpToDefinition() const override;
-	virtual void JumpToDefinition() const override;
-	virtual FText GetTooltipText() const override;
+	UE_API virtual void AllocateDefaultPins() override;
+	UE_API virtual bool IsCompatibleWithGraph(const UEdGraph* TargetGraph) const override;
+	UE_API virtual FSlateIcon GetIconAndTint(FLinearColor& OutColor) const override;
+	UE_API virtual FLinearColor GetNodeTitleColor() const override;
+	UE_API virtual FText GetNodeTitle(ENodeTitleType::Type TitleType) const override;
+	UE_API virtual void ExpandNode(FKismetCompilerContext& CompilerContext, UEdGraph* SourceGraph) override;
+	UE_API virtual bool CanJumpToDefinition() const override;
+	UE_API virtual void JumpToDefinition() const override;
+	UE_API virtual FText GetTooltipText() const override;
 	//~End UEdGraphNode interface.
 
 	virtual bool IsNodePure() const { return true; }
 	virtual bool IsNodeSafeToIgnore() const override { return true; }
-	virtual void ReallocatePinsDuringReconstruction(TArray<UEdGraphPin*>& OldPins) override;
-	virtual class FNodeHandlingFunctor* CreateNodeHandler(class FKismetCompilerContext& CompilerContext) const override;
-	virtual void GetMenuActions(FBlueprintActionDatabaseRegistrar& ActionRegistrar) const override;
-	virtual void GetNodeAttributes(TArray<TKeyValuePair<FString, FString>>& OutNodeAttributes) const override;
-	virtual FText GetMenuCategory() const override;
+	UE_API virtual void ReallocatePinsDuringReconstruction(TArray<UEdGraphPin*>& OldPins) override;
+	UE_API virtual class FNodeHandlingFunctor* CreateNodeHandler(class FKismetCompilerContext& CompilerContext) const override;
+	UE_API virtual void GetMenuActions(FBlueprintActionDatabaseRegistrar& ActionRegistrar) const override;
+	UE_API virtual void GetNodeAttributes(TArray<TKeyValuePair<FString, FString>>& OutNodeAttributes) const override;
+	UE_API virtual FText GetMenuCategory() const override;
 
 	/** Get the blueprint input pin */
-	UEdGraphPin* GetClassPin(const TArray< UEdGraphPin* >* InPinsToSearch = nullptr) const;
+	UE_API UEdGraphPin* GetClassPin(const TArray< UEdGraphPin* >* InPinsToSearch = nullptr) const;
 	/** Get the world context input pin, can return NULL */
-	UEdGraphPin* GetWorldContextPin() const;
+	UE_API UEdGraphPin* GetWorldContextPin() const;
 	/** Get the result output pin */
-	UEdGraphPin* GetResultPin() const;
+	UE_API UEdGraphPin* GetResultPin() const;
 
 	virtual bool ShouldDrawCompact() const override { return true; }
 
@@ -125,4 +127,6 @@ public:
 	virtual void GetMenuActions(FBlueprintActionDatabaseRegistrar& ActionRegistrar) const override;
 	virtual FText GetMenuCategory() const override;
 
-};
\ No newline at end of file
+};
+
+#undef UE_API
diff --git a/Engine/Source/Editor/BlueprintGraph/Classes/K2Node_InputAction.h b/Engine/Source/Editor/BlueprintGraph/Classes/K2Node_InputAction.h
--- a/Engine/Source/Editor/BlueprintGraph/Classes/K2Node_InputAction.h
+++ b/Engine/Source/Editor/BlueprintGraph/Classes/K2Node_InputAction.h
@@ -21,6 +21,8 @@
 
 #include "K2Node_InputAction.generated.h"
 
+#define UE_API BLUEPRINTGRAPH_API
+
 class FBlueprintActionDatabaseRegistrar;
 class FKismetCompilerContext;
 class UEdGraph;
@@ -28,8 +30,8 @@ class UEdGraphPin;
 class UObject;
 struct FEdGraphSchemaAction;
 
-UCLASS()
-class BLUEPRINTGRAPH_API UK2Node_InputAction : public UK2Node, public IK2Node_EventNodeInterface
+UCLASS(MinimalAPI)
+class UK2Node_InputAction : public UK2Node, public IK2Node_EventNodeInterface
 {
 	GENERATED_UCLASS_BODY()
 
@@ -49,39 +51,41 @@ class BLUEPRINTGRAPH_API UK2Node_InputAction : public UK2Node, public IK2Node_Ev
 	uint32 bOverrideParentBinding:1;
 
 	//~ Begin UObject Interface
-	virtual void PostLoad() override;
+	UE_API virtual void PostLoad() override;
 	//~ End UObject Interface
 
 	//~ Begin UEdGraphNode Interface.
-	virtual void AllocateDefaultPins() override;
-	virtual FLinearColor GetNodeTitleColor() const override;
-	virtual FText GetNodeTitle(ENodeTitleType::Type TitleType) const override;
-	virtual FText GetTooltipText() const override;
-	virtual FSlateIcon GetIconAndTint(FLinearColor& OutColor) const override;
-	virtual bool IsCompatibleWithGraph(UEdGraph const* Graph) const override;
+	UE_API virtual void AllocateDefaultPins() override;
+	UE_API virtual FLinearColor GetNodeTitleColor() const override;
+	UE_API virtual FText GetNodeTitle(ENodeTitleType::Type TitleType) const override;
+	UE_API virtual FText GetTooltipText() const override;
+	UE_API virtual FSlateIcon GetIconAndTint(FLinearColor& OutColor) const override;
+	UE_API virtual bool IsCompatibleWithGraph(UEdGraph const* Graph) const override;
 	//~ End UEdGraphNode Interface.
 
 	//~ Begin UK2Node Interface
-	virtual void ValidateNodeDuringCompilation(class FCompilerResultsLog& MessageLog) const override;
+	UE_API virtual void ValidateNodeDuringCompilation(class FCompilerResultsLog& MessageLog) const override;
 	virtual bool ShouldShowNodeProperties() const override { return true; }
-	virtual void ExpandNode(FKismetCompilerContext& CompilerContext, UEdGraph* SourceGraph) override;
-	virtual void GetMenuActions(FBlueprintActionDatabaseRegistrar& ActionRegistrar) const override;
-	virtual FText GetMenuCategory() const override;
-	virtual FBlueprintNodeSignature GetSignature() const override;
+	UE_API virtual void ExpandNode(FKismetCompilerContext& CompilerContext, UEdGraph* SourceGraph) override;
+	UE_API virtual void GetMenuActions(FBlueprintActionDatabaseRegistrar& ActionRegistrar) const override;
+	UE_API virtual FText GetMenuCategory() const override;
+	UE_API virtual FBlueprintNodeSignature GetSignature() const override;
 	//~ End UK2Node Interface
 
 	//~ Begin IK2Node_EventNodeInterface Interface.
-	virtual TSharedPtr<FEdGraphSchemaAction> GetEventNodeAction(const FText& ActionCategory) override;
+	UE_API virtual TSharedPtr<FEdGraphSchemaAction> GetEventNodeAction(const FText& ActionCategory) override;
 	//~ End IK2Node_EventNodeInterface Interface.
 
 	/** Get the 'pressed' input pin */
-	UEdGraphPin* GetPressedPin() const;
+	UE_API UEdGraphPin* GetPressedPin() const;
 
 	/** Get the 'released' input pin */
-	UEdGraphPin* GetReleasedPin() const;
+	UE_API UEdGraphPin* GetReleasedPin() const;
 
 private:
 	/** Constructing FText strings can be costly, so we cache the node's title/tooltip */
 	FNodeTextCache CachedTooltip;
 	FNodeTextCache CachedNodeTitle;
 };
+
+#undef UE_API
diff --git a/Engine/Source/Editor/BlueprintGraph/Classes/K2Node_InputActionEvent.h b/Engine/Source/Editor/BlueprintGraph/Classes/K2Node_InputActionEvent.h
--- a/Engine/Source/Editor/BlueprintGraph/Classes/K2Node_InputActionEvent.h
+++ b/Engine/Source/Editor/BlueprintGraph/Classes/K2Node_InputActionEvent.h
@@ -14,12 +14,14 @@
 
 #include "K2Node_InputActionEvent.generated.h"
 
+#define UE_API BLUEPRINTGRAPH_API
+
 class UClass;
 class UDynamicBlueprintBinding;
 class UObject;
 
-UCLASS()
-class BLUEPRINTGRAPH_API UK2Node_InputActionEvent : public UK2Node_Event
+UCLASS(MinimalAPI)
+class UK2Node_InputActionEvent : public UK2Node_Event
 {
 	GENERATED_UCLASS_BODY()
 
@@ -39,7 +41,9 @@ class BLUEPRINTGRAPH_API UK2Node_InputActionEvent : public UK2Node_Event
 	uint32 bOverrideParentBinding:1;
 
 	//~ Begin UK2Node Interface
-	virtual UClass* GetDynamicBindingClass() const override;
-	virtual void RegisterDynamicBinding(UDynamicBlueprintBinding* BindingObject) const override;
+	UE_API virtual UClass* GetDynamicBindingClass() const override;
+	UE_API virtual void RegisterDynamicBinding(UDynamicBlueprintBinding* BindingObject) const override;
 	//~ End UK2Node Interface
 };
+
+#undef UE_API
diff --git a/Engine/Source/Editor/BlueprintGraph/Classes/K2Node_InputAxisEvent.h b/Engine/Source/Editor/BlueprintGraph/Classes/K2Node_InputAxisEvent.h
--- a/Engine/Source/Editor/BlueprintGraph/Classes/K2Node_InputAxisEvent.h
+++ b/Engine/Source/Editor/BlueprintGraph/Classes/K2Node_InputAxisEvent.h
@@ -15,6 +15,8 @@
 
 #include "K2Node_InputAxisEvent.generated.h"
 
+#define UE_API BLUEPRINTGRAPH_API
+
 class FArchive;
 class FBlueprintActionDatabaseRegistrar;
 class UClass;
@@ -22,8 +24,8 @@ class UDynamicBlueprintBinding;
 class UEdGraph;
 class UObject;
 
-UCLASS()
-class BLUEPRINTGRAPH_API UK2Node_InputAxisEvent : public UK2Node_Event
+UCLASS(MinimalAPI)
+class UK2Node_InputAxisEvent : public UK2Node_Event
 {
 	GENERATED_UCLASS_BODY()
 
@@ -43,31 +45,33 @@ class BLUEPRINTGRAPH_API UK2Node_InputAxisEvent : public UK2Node_Event
 	uint32 bOverrideParentBinding:1;
 
 	//~ Begin UObject Interface
-	virtual void PostLoad() override;
-	virtual void Serialize(FArchive& Ar) override;
+	UE_API virtual void PostLoad() override;
+	UE_API virtual void Serialize(FArchive& Ar) override;
 	//~ End UObject Interface
 
 	//~ Begin EdGraphNode Interface
-	virtual FText GetNodeTitle(ENodeTitleType::Type TitleType) const override;
-	virtual FText GetTooltipText() const override;
-	virtual bool IsCompatibleWithGraph(const UEdGraph* TargetGraph) const override;
+	UE_API virtual FText GetNodeTitle(ENodeTitleType::Type TitleType) const override;
+	UE_API virtual FText GetTooltipText() const override;
+	UE_API virtual bool IsCompatibleWithGraph(const UEdGraph* TargetGraph) const override;
 	//~ End EdGraphNode Interface
 
 	//~ Begin UK2Node Interface
-	virtual void ValidateNodeDuringCompilation(class FCompilerResultsLog& MessageLog) const override;
+	UE_API virtual void ValidateNodeDuringCompilation(class FCompilerResultsLog& MessageLog) const override;
 	virtual bool ShouldShowNodeProperties() const override { return true; }
-	virtual UClass* GetDynamicBindingClass() const override;
-	virtual void RegisterDynamicBinding(UDynamicBlueprintBinding* BindingObject) const override;
-	virtual void GetMenuActions(FBlueprintActionDatabaseRegistrar& ActionRegistrar) const override;
-	virtual FText GetMenuCategory() const override;
-	virtual FBlueprintNodeSignature GetSignature() const override;
-	virtual void ExpandNode(class FKismetCompilerContext& CompilerContext, UEdGraph* SourceGraph) override;
+	UE_API virtual UClass* GetDynamicBindingClass() const override;
+	UE_API virtual void RegisterDynamicBinding(UDynamicBlueprintBinding* BindingObject) const override;
+	UE_API virtual void GetMenuActions(FBlueprintActionDatabaseRegistrar& ActionRegistrar) const override;
+	UE_API virtual FText GetMenuCategory() const override;
+	UE_API virtual FBlueprintNodeSignature GetSignature() const override;
+	UE_API virtual void ExpandNode(class FKismetCompilerContext& CompilerContext, UEdGraph* SourceGraph) override;
 	//~ End UK2Node Interface
 
-	void Initialize(const FName AxisName);
+	UE_API void Initialize(const FName AxisName);
 
 private:
 	/** Constructing FText strings can be costly, so we cache the node's title/tooltip */
 	FNodeTextCache CachedTooltip;
 	FNodeTextCache CachedNodeTitle;
 };
+
+#undef UE_API
diff --git a/Engine/Source/Editor/BlueprintGraph/Classes/K2Node_InputAxisKeyEvent.h b/Engine/Source/Editor/BlueprintGraph/Classes/K2Node_InputAxisKeyEvent.h
--- a/Engine/Source/Editor/BlueprintGraph/Classes/K2Node_InputAxisKeyEvent.h
+++ b/Engine/Source/Editor/BlueprintGraph/Classes/K2Node_InputAxisKeyEvent.h
@@ -16,6 +16,8 @@
 
 #include "K2Node_InputAxisKeyEvent.generated.h"
 
+#define UE_API BLUEPRINTGRAPH_API
+
 class FArchive;
 class FBlueprintActionDatabaseRegistrar;
 class UClass;
@@ -24,8 +26,8 @@ class UEdGraph;
 class UObject;
 struct FLinearColor;
 
-UCLASS()
-class BLUEPRINTGRAPH_API UK2Node_InputAxisKeyEvent : public UK2Node_Event
+UCLASS(MinimalAPI)
+class UK2Node_InputAxisKeyEvent : public UK2Node_Event
 {
 	GENERATED_UCLASS_BODY()
 
@@ -45,29 +47,31 @@ class BLUEPRINTGRAPH_API UK2Node_InputAxisKeyEvent : public UK2Node_Event
 	uint32 bOverrideParentBinding : 1;
 
 	//~ Begin UObject Interface
-	virtual void Serialize(FArchive& Ar) override;
+	UE_API virtual void Serialize(FArchive& Ar) override;
 	//~ End UObject Interface
 
 	//~ Begin EdGraphNode Interface
-	virtual FText GetNodeTitle(ENodeTitleType::Type TitleType) const override;
-	virtual FText GetTooltipText() const override;
-	virtual FSlateIcon GetIconAndTint(FLinearColor& OutColor) const override;
-	virtual bool IsCompatibleWithGraph(const UEdGraph* TargetGraph) const override;
+	UE_API virtual FText GetNodeTitle(ENodeTitleType::Type TitleType) const override;
+	UE_API virtual FText GetTooltipText() const override;
+	UE_API virtual FSlateIcon GetIconAndTint(FLinearColor& OutColor) const override;
+	UE_API virtual bool IsCompatibleWithGraph(const UEdGraph* TargetGraph) const override;
 	//~ End EdGraphNode Interface
 
 	//~ Begin UK2Node Interface
-	virtual void ValidateNodeDuringCompilation(class FCompilerResultsLog& MessageLog) const override;
+	UE_API virtual void ValidateNodeDuringCompilation(class FCompilerResultsLog& MessageLog) const override;
 	virtual bool ShouldShowNodeProperties() const override{ return true; }
-	virtual UClass* GetDynamicBindingClass() const override;
-	virtual void RegisterDynamicBinding(UDynamicBlueprintBinding* BindingObject) const override;
-	virtual void GetMenuActions(FBlueprintActionDatabaseRegistrar& ActionRegistrar) const override;
-	virtual FText GetMenuCategory() const override;
-	virtual FBlueprintNodeSignature GetSignature() const override;
+	UE_API virtual UClass* GetDynamicBindingClass() const override;
+	UE_API virtual void RegisterDynamicBinding(UDynamicBlueprintBinding* BindingObject) const override;
+	UE_API virtual void GetMenuActions(FBlueprintActionDatabaseRegistrar& ActionRegistrar) const override;
+	UE_API virtual FText GetMenuCategory() const override;
+	UE_API virtual FBlueprintNodeSignature GetSignature() const override;
 	//~ End UK2Node Interface
 
-	void Initialize(const FKey AxisKey);
+	UE_API void Initialize(const FKey AxisKey);
 
 private:
 	/** Constructing FText strings can be costly, so we cache the node's tooltip */
 	FNodeTextCache CachedTooltip;
 };
+
+#undef UE_API
