Updated headers using UnrealCodeFixup to make sure dllstorage is on methods/staticvars instead of types. Part 4/n

diff --git a/Engine/Source/Editor/BlueprintGraph/Classes/K2Node_InputKey.h b/Engine/Source/Editor/BlueprintGraph/Classes/K2Node_InputKey.h
--- a/Engine/Source/Editor/BlueprintGraph/Classes/K2Node_InputKey.h
+++ b/Engine/Source/Editor/BlueprintGraph/Classes/K2Node_InputKey.h
@@ -22,6 +22,8 @@
 
 #include "K2Node_InputKey.generated.h"
 
+#define UE_API BLUEPRINTGRAPH_API
+
 class FBlueprintActionDatabaseRegistrar;
 class FKismetCompilerContext;
 class UEdGraph;
@@ -29,8 +31,8 @@ class UEdGraphPin;
 class UObject;
 struct FEdGraphSchemaAction;
 
-UCLASS()
-class BLUEPRINTGRAPH_API UK2Node_InputKey : public UK2Node, public IK2Node_EventNodeInterface
+UCLASS(MinimalAPI)
+class UK2Node_InputKey : public UK2Node, public IK2Node_EventNodeInterface
 {
 	GENERATED_UCLASS_BODY()
 
@@ -67,47 +69,49 @@ class BLUEPRINTGRAPH_API UK2Node_InputKey : public UK2Node, public IK2Node_Event
 	uint32 bCommand:1;
 
 	//~ Begin UObject Interface
-	virtual void PostLoad() override;
-	virtual void PostEditChangeProperty(struct FPropertyChangedEvent& PropertyChangedEvent) override;
+	UE_API virtual void PostLoad() override;
+	UE_API virtual void PostEditChangeProperty(struct FPropertyChangedEvent& PropertyChangedEvent) override;
 	//~ End UObject Interface
 
 	//~ Begin UK2Node Interface.
 	virtual bool ShouldShowNodeProperties() const override { return true; }
-	virtual void ValidateNodeDuringCompilation(class FCompilerResultsLog& MessageLog) const override;
-	virtual void ExpandNode(FKismetCompilerContext& CompilerContext, UEdGraph* SourceGraph) override;
-	virtual void GetMenuActions(FBlueprintActionDatabaseRegistrar& ActionRegistrar) const override;
-	virtual FText GetMenuCategory() const override;
-	virtual FBlueprintNodeSignature GetSignature() const override;
+	UE_API virtual void ValidateNodeDuringCompilation(class FCompilerResultsLog& MessageLog) const override;
+	UE_API virtual void ExpandNode(FKismetCompilerContext& CompilerContext, UEdGraph* SourceGraph) override;
+	UE_API virtual void GetMenuActions(FBlueprintActionDatabaseRegistrar& ActionRegistrar) const override;
+	UE_API virtual FText GetMenuCategory() const override;
+	UE_API virtual FBlueprintNodeSignature GetSignature() const override;
 	//~ End UK2Node Interface
 
 	//~ Begin UEdGraphNode Interface.
-	virtual void AllocateDefaultPins() override;
-	virtual FLinearColor GetNodeTitleColor() const override;
-	virtual FText GetNodeTitle(ENodeTitleType::Type TitleType) const override;
-	virtual FText GetTooltipText() const override;
-	virtual FSlateIcon GetIconAndTint(FLinearColor& OutColor) const override;
-	virtual bool IsCompatibleWithGraph(UEdGraph const* Graph) const override;
+	UE_API virtual void AllocateDefaultPins() override;
+	UE_API virtual FLinearColor GetNodeTitleColor() const override;
+	UE_API virtual FText GetNodeTitle(ENodeTitleType::Type TitleType) const override;
+	UE_API virtual FText GetTooltipText() const override;
+	UE_API virtual FSlateIcon GetIconAndTint(FLinearColor& OutColor) const override;
+	UE_API virtual bool IsCompatibleWithGraph(UEdGraph const* Graph) const override;
 	//~ End UEdGraphNode Interface.
 
 	//~ Begin IK2Node_EventNodeInterface Interface.
-	virtual TSharedPtr<FEdGraphSchemaAction> GetEventNodeAction(const FText& ActionCategory) override;
+	UE_API virtual TSharedPtr<FEdGraphSchemaAction> GetEventNodeAction(const FText& ActionCategory) override;
 	//~ End IK2Node_EventNodeInterface Interface.
 
 
-	FText GetModifierText() const;
+	UE_API FText GetModifierText() const;
 
-	FName GetModifierName() const;
+	UE_API FName GetModifierName() const;
 
-	FText GetKeyText() const;
+	UE_API FText GetKeyText() const;
 
 	/** Get the 'pressed' input pin */
-	UEdGraphPin* GetPressedPin() const;
+	UE_API UEdGraphPin* GetPressedPin() const;
 
 	/** Get the 'released' input pin */
-	UEdGraphPin* GetReleasedPin() const;
+	UE_API UEdGraphPin* GetReleasedPin() const;
 
 private:
 	/** Constructing FText strings can be costly, so we cache the node's title/tooltip */
 	FNodeTextCache CachedTooltip;
 	FNodeTextCache CachedNodeTitle;
 };
+
+#undef UE_API
diff --git a/Engine/Source/Editor/BlueprintGraph/Classes/K2Node_InputKeyEvent.h b/Engine/Source/Editor/BlueprintGraph/Classes/K2Node_InputKeyEvent.h
--- a/Engine/Source/Editor/BlueprintGraph/Classes/K2Node_InputKeyEvent.h
+++ b/Engine/Source/Editor/BlueprintGraph/Classes/K2Node_InputKeyEvent.h
@@ -14,12 +14,14 @@
 
 #include "K2Node_InputKeyEvent.generated.h"
 
+#define UE_API BLUEPRINTGRAPH_API
+
 class UClass;
 class UDynamicBlueprintBinding;
 class UObject;
 
-UCLASS()
-class BLUEPRINTGRAPH_API UK2Node_InputKeyEvent : public UK2Node_Event
+UCLASS(MinimalAPI)
+class UK2Node_InputKeyEvent : public UK2Node_Event
 {
 	GENERATED_UCLASS_BODY()
 
@@ -39,7 +41,9 @@ class BLUEPRINTGRAPH_API UK2Node_InputKeyEvent : public UK2Node_Event
 	uint32 bOverrideParentBinding:1;
 
 	//~ Begin UK2Node Interface
-	virtual UClass* GetDynamicBindingClass() const override;
-	virtual void RegisterDynamicBinding(UDynamicBlueprintBinding* BindingObject) const override;
+	UE_API virtual UClass* GetDynamicBindingClass() const override;
+	UE_API virtual void RegisterDynamicBinding(UDynamicBlueprintBinding* BindingObject) const override;
 	//~ End UK2Node Interface
 };
+
+#undef UE_API
diff --git a/Engine/Source/Editor/BlueprintGraph/Classes/K2Node_Literal.h b/Engine/Source/Editor/BlueprintGraph/Classes/K2Node_Literal.h
--- a/Engine/Source/Editor/BlueprintGraph/Classes/K2Node_Literal.h
+++ b/Engine/Source/Editor/BlueprintGraph/Classes/K2Node_Literal.h
@@ -61,6 +61,6 @@ public:
 	BLUEPRINTGRAPH_API void SetObjectRef(UObject* NewValue);
 
 	/** Gets the referenced object */
-	BLUEPRINTGRAPH_API UObject* GetObjectRef() const { return ObjectRef; }
+	UObject* GetObjectRef() const { return ObjectRef; }
 };
 
diff --git a/Engine/Source/Editor/BlueprintGraph/Classes/K2Node_LoadAsset.h b/Engine/Source/Editor/BlueprintGraph/Classes/K2Node_LoadAsset.h
--- a/Engine/Source/Editor/BlueprintGraph/Classes/K2Node_LoadAsset.h
+++ b/Engine/Source/Editor/BlueprintGraph/Classes/K2Node_LoadAsset.h
@@ -14,41 +14,43 @@
 
 #include "K2Node_LoadAsset.generated.h"
 
+#define UE_API BLUEPRINTGRAPH_API
+
 class FBlueprintActionDatabaseRegistrar;
 class UEdGraph;
 class UEdGraphPin;
 class UObject;
 
-UCLASS()
-class BLUEPRINTGRAPH_API UK2Node_LoadAsset : public UK2Node
+UCLASS(MinimalAPI)
+class UK2Node_LoadAsset : public UK2Node
 {
 	GENERATED_BODY()
 public:
 	// UEdGraphNode interface
-	virtual void AllocateDefaultPins() override;
-	virtual FText GetTooltipText() const override;
-	virtual FText GetNodeTitle(ENodeTitleType::Type TitleType) const override;
-	virtual bool IsCompatibleWithGraph(const UEdGraph* TargetGraph) const override;
+	UE_API virtual void AllocateDefaultPins() override;
+	UE_API virtual FText GetTooltipText() const override;
+	UE_API virtual FText GetNodeTitle(ENodeTitleType::Type TitleType) const override;
+	UE_API virtual bool IsCompatibleWithGraph(const UEdGraph* TargetGraph) const override;
 	// End of UEdGraphNode interface
 
 	// UK2Node interface
 	virtual bool IsNodePure() const override { return false; }
-	virtual void ExpandNode(class FKismetCompilerContext& CompilerContext, UEdGraph* SourceGraph) override;
-	virtual FName GetCornerIcon() const override;
-	virtual void GetMenuActions(FBlueprintActionDatabaseRegistrar& ActionRegistrar) const override;
-	virtual FText GetMenuCategory() const override;
+	UE_API virtual void ExpandNode(class FKismetCompilerContext& CompilerContext, UEdGraph* SourceGraph) override;
+	UE_API virtual FName GetCornerIcon() const override;
+	UE_API virtual void GetMenuActions(FBlueprintActionDatabaseRegistrar& ActionRegistrar) const override;
+	UE_API virtual FText GetMenuCategory() const override;
 	virtual bool NodeCausesStructuralBlueprintChange() const { return true; }
-	virtual void ReallocatePinsDuringReconstruction(TArray<UEdGraphPin*>& OldPins) override;
+	UE_API virtual void ReallocatePinsDuringReconstruction(TArray<UEdGraphPin*>& OldPins) override;
 	// End of UK2Node interface
 
-	virtual const FName& GetInputPinName() const;
-	virtual const FName& GetOutputPinName() const;
+	UE_API virtual const FName& GetInputPinName() const;
+	UE_API virtual const FName& GetOutputPinName() const;
 
 protected:
-	virtual FName NativeFunctionName() const;
+	UE_API virtual FName NativeFunctionName() const;
 
-	virtual const FName& GetInputCategory() const;
-	virtual const FName& GetOutputCategory() const;
+	UE_API virtual const FName& GetInputCategory() const;
+	UE_API virtual const FName& GetOutputCategory() const;
 };
 
 UCLASS(MinimalAPI)
@@ -70,3 +72,5 @@ protected:
 	virtual const FName& GetInputPinName() const override;
 	virtual const FName& GetOutputPinName() const override;
 };
+
+#undef UE_API
diff --git a/Engine/Source/Editor/BlueprintGraph/Classes/K2Node_MakeContainer.h b/Engine/Source/Editor/BlueprintGraph/Classes/K2Node_MakeContainer.h
--- a/Engine/Source/Editor/BlueprintGraph/Classes/K2Node_MakeContainer.h
+++ b/Engine/Source/Editor/BlueprintGraph/Classes/K2Node_MakeContainer.h
@@ -20,6 +20,8 @@
 
 #include "K2Node_MakeContainer.generated.h"
 
+#define UE_API BLUEPRINTGRAPH_API
+
 class FBlueprintActionDatabaseRegistrar;
 class FKismetCompilerContext;
 class FString;
@@ -43,8 +45,8 @@ protected:
 };
 
 
-UCLASS(abstract)
-class BLUEPRINTGRAPH_API UK2Node_MakeContainer : public UK2Node, public IK2Node_AddPinInterface
+UCLASS(MinimalAPI, abstract)
+class UK2Node_MakeContainer : public UK2Node, public IK2Node_AddPinInterface
 {
 	GENERATED_UCLASS_BODY()
 
@@ -53,50 +55,52 @@ class BLUEPRINTGRAPH_API UK2Node_MakeContainer : public UK2Node, public IK2Node_
 	int32 NumInputs;
 
 public:
-	void RemoveInputPin(UEdGraphPin* Pin);
+	UE_API void RemoveInputPin(UEdGraphPin* Pin);
 
-	UEdGraphPin* GetOutputPin() const;
+	UE_API UEdGraphPin* GetOutputPin() const;
 
 	/** returns a reference to the output array pin of this node, which is responsible for defining the type */
 	virtual FName GetOutputPinName() const PURE_VIRTUAL(UK2Node_MakeContainer::GetOutputPinName, return NAME_None;);
-	virtual FName GetPinName(int32 PinIndex) const;
-	virtual void GetKeyAndValuePins(TArray<UEdGraphPin*>& KeyPins, TArray<UEdGraphPin*>& ValuePins) const;
+	UE_API virtual FName GetPinName(int32 PinIndex) const;
+	UE_API virtual void GetKeyAndValuePins(TArray<UEdGraphPin*>& KeyPins, TArray<UEdGraphPin*>& ValuePins) const;
 
 public:
 	// UEdGraphNode interface
-	virtual void AllocateDefaultPins() override;
-	virtual void PostReconstructNode() override;
+	UE_API virtual void AllocateDefaultPins() override;
+	UE_API virtual void PostReconstructNode() override;
 	// End of UEdGraphNode interface
 
 	// UK2Node interface
 	virtual bool IsNodePure() const override { return true; }
-	virtual void NotifyPinConnectionListChanged(UEdGraphPin* Pin) override;
+	UE_API virtual void NotifyPinConnectionListChanged(UEdGraphPin* Pin) override;
 	virtual class FNodeHandlingFunctor* CreateNodeHandler(class FKismetCompilerContext& CompilerContext) const override PURE_VIRTUAL(UK2Node_MakeContainer::CreateNodeHandler, return nullptr;);
-	virtual void GetMenuActions(FBlueprintActionDatabaseRegistrar& ActionRegistrar) const override;
+	UE_API virtual void GetMenuActions(FBlueprintActionDatabaseRegistrar& ActionRegistrar) const override;
 	virtual int32 GetNodeRefreshPriority() const override { return EBaseNodeRefreshPriority::Low_UsesDependentWildcard; }
-	virtual void ReallocatePinsDuringReconstruction(TArray<UEdGraphPin*>& OldPins);
-	virtual bool IsConnectionDisallowed(const UEdGraphPin* MyPin, const UEdGraphPin* OtherPin, FString& OutReason) const override;
+	UE_API virtual void ReallocatePinsDuringReconstruction(TArray<UEdGraphPin*>& OldPins);
+	UE_API virtual bool IsConnectionDisallowed(const UEdGraphPin* MyPin, const UEdGraphPin* OtherPin, FString& OutReason) const override;
 	// End of UK2Node interface
 
 	// IK2Node_AddPinInterface interface
-	virtual void AddInputPin() override;
+	UE_API virtual void AddInputPin() override;
 	// End of IK2Node_AddPinInterface interface
 
 protected:
 	friend class FKismetCompilerContext;
 
 	/** If needed, will clear all pins to be wildcards */
-	void ClearPinTypeToWildcard();
+	UE_API void ClearPinTypeToWildcard();
 
-	bool CanResetToWildcard() const;
+	UE_API bool CanResetToWildcard() const;
 
 	/** Helper function for context menu add pin to ensure transaction is set up correctly. */
-	void InteractiveAddInputPin();
+	UE_API void InteractiveAddInputPin();
 
 	/** Propagates the pin type from the output (set) pin to the inputs, to make sure types are consistent */
-	void PropagatePinType();
+	UE_API void PropagatePinType();
 
-	void SyncPinNames();
+	UE_API void SyncPinNames();
 
 	EPinContainerType ContainerType;
 };
+
+#undef UE_API
diff --git a/Engine/Source/Editor/BlueprintGraph/Classes/K2Node_MakeMap.h b/Engine/Source/Editor/BlueprintGraph/Classes/K2Node_MakeMap.h
--- a/Engine/Source/Editor/BlueprintGraph/Classes/K2Node_MakeMap.h
+++ b/Engine/Source/Editor/BlueprintGraph/Classes/K2Node_MakeMap.h
@@ -15,38 +15,42 @@
 
 #include "K2Node_MakeMap.generated.h"
 
+#define UE_API BLUEPRINTGRAPH_API
+
 class UEdGraphPin;
 class UObject;
 struct FLinearColor;
 
-UCLASS()
-class BLUEPRINTGRAPH_API UK2Node_MakeMap : public UK2Node_MakeContainer
+UCLASS(MinimalAPI)
+class UK2Node_MakeMap : public UK2Node_MakeContainer
 {
 	GENERATED_UCLASS_BODY()
 
 public:
 	// UEdGraphNode interface
-	virtual void AllocateDefaultPins() override;
-	virtual FText GetNodeTitle(ENodeTitleType::Type TitleType) const override;
-	virtual FText GetTooltipText() const override;
-	virtual FSlateIcon GetIconAndTint(FLinearColor& OutColor) const override;
-	virtual void ValidateNodeDuringCompilation(class FCompilerResultsLog& MessageLog) const override;
+	UE_API virtual void AllocateDefaultPins() override;
+	UE_API virtual FText GetNodeTitle(ENodeTitleType::Type TitleType) const override;
+	UE_API virtual FText GetTooltipText() const override;
+	UE_API virtual FSlateIcon GetIconAndTint(FLinearColor& OutColor) const override;
+	UE_API virtual void ValidateNodeDuringCompilation(class FCompilerResultsLog& MessageLog) const override;
 	// End of UEdGraphNode interface
 
 	// UK2Node interface
-	virtual void GetNodeContextMenuActions(class UToolMenu* Menu, class UGraphNodeContextMenuContext* Context) const override;
+	UE_API virtual void GetNodeContextMenuActions(class UToolMenu* Menu, class UGraphNodeContextMenuContext* Context) const override;
 	virtual bool IncludeParentNodeContextMenu() const override { return true; }
-	virtual class FNodeHandlingFunctor* CreateNodeHandler(class FKismetCompilerContext& CompilerContext) const override;
-	virtual FText GetMenuCategory() const override;
+	UE_API virtual class FNodeHandlingFunctor* CreateNodeHandler(class FKismetCompilerContext& CompilerContext) const override;
+	UE_API virtual FText GetMenuCategory() const override;
 	// End of UK2Node interface
 
 	// UK2Node_MakeContainer interface
-	virtual FName GetOutputPinName() const override;
-	virtual FName GetPinName(int32 PinIndex) const override;
-	virtual void GetKeyAndValuePins(TArray<UEdGraphPin*>& KeyPins, TArray<UEdGraphPin*>& ValuePins) const override;
+	UE_API virtual FName GetOutputPinName() const override;
+	UE_API virtual FName GetPinName(int32 PinIndex) const override;
+	UE_API virtual void GetKeyAndValuePins(TArray<UEdGraphPin*>& KeyPins, TArray<UEdGraphPin*>& ValuePins) const override;
 	// UK2Node_MakeContainer interface
 
 	// IK2Node_AddPinInterface interface
-	virtual void AddInputPin() override;
+	UE_API virtual void AddInputPin() override;
 	// End of IK2Node_AddPinInterface interface
 };
+
+#undef UE_API
diff --git a/Engine/Source/Editor/BlueprintGraph/Classes/K2Node_MathExpression.h b/Engine/Source/Editor/BlueprintGraph/Classes/K2Node_MathExpression.h
--- a/Engine/Source/Editor/BlueprintGraph/Classes/K2Node_MathExpression.h
+++ b/Engine/Source/Editor/BlueprintGraph/Classes/K2Node_MathExpression.h
@@ -16,6 +16,8 @@
 
 #include "K2Node_MathExpression.generated.h"
 
+#define UE_API BLUEPRINTGRAPH_API
+
 class FArchive;
 class FBlueprintActionDatabaseRegistrar;
 class INameValidatorInterface;
@@ -29,8 +31,8 @@ struct FEdGraphPinType;
 * series of math nodes. If it is unsuccessful, then it generates a series of
 * actionable errors.
 */
-UCLASS()
-class BLUEPRINTGRAPH_API UK2Node_MathExpression : public UK2Node_Composite
+UCLASS(MinimalAPI)
+class UK2Node_MathExpression : public UK2Node_Composite
 {
 	GENERATED_UCLASS_BODY()
 
@@ -45,24 +47,24 @@ public:
 
 public:
 	//~ Begin UObject Interface
-	virtual void Serialize(FArchive& Ar) override;
-	virtual void PostEditChangeProperty(struct FPropertyChangedEvent& PropertyChangedEvent) override;
+	UE_API virtual void Serialize(FArchive& Ar) override;
+	UE_API virtual void PostEditChangeProperty(struct FPropertyChangedEvent& PropertyChangedEvent) override;
 	//~ End UObject Interface
 
 	//~ Begin UEdGraphNode Interface
-	virtual TSharedPtr<class INameValidatorInterface> MakeNameValidator() const override;
-	virtual void OnRenameNode(const FString& NewName) override;
-	virtual FText GetNodeTitle(ENodeTitleType::Type TitleType) const override;
-	virtual void PostPlacedNewNode() override;
-	virtual void ReconstructNode() override;
-	virtual void FindDiffs(class UEdGraphNode* OtherNode, struct FDiffResults& Results )  override;
+	UE_API virtual TSharedPtr<class INameValidatorInterface> MakeNameValidator() const override;
+	UE_API virtual void OnRenameNode(const FString& NewName) override;
+	UE_API virtual FText GetNodeTitle(ENodeTitleType::Type TitleType) const override;
+	UE_API virtual void PostPlacedNewNode() override;
+	UE_API virtual void ReconstructNode() override;
+	UE_API virtual void FindDiffs(class UEdGraphNode* OtherNode, struct FDiffResults& Results )  override;
 	virtual bool ShouldMergeChildGraphs() const override { return ShouldExpandInsteadCompile(); }
 	//~ End UEdGraphNode Interface
 
 	//~ Begin UK2Node Interface
-	virtual void ValidateNodeDuringCompilation(class FCompilerResultsLog& MessageLog) const override;
-	virtual void GetMenuActions(FBlueprintActionDatabaseRegistrar& ActionRegistrar) const override;
-	virtual class FNodeHandlingFunctor* CreateNodeHandler(class FKismetCompilerContext& CompilerContext) const override;
+	UE_API virtual void ValidateNodeDuringCompilation(class FCompilerResultsLog& MessageLog) const override;
+	UE_API virtual void GetMenuActions(FBlueprintActionDatabaseRegistrar& ActionRegistrar) const override;
+	UE_API virtual class FNodeHandlingFunctor* CreateNodeHandler(class FKismetCompilerContext& CompilerContext) const override;
 	virtual bool IsNodePure() const { return !ShouldExpandInsteadCompile(); }
 	//~ End UK2Node Interface
 
@@ -72,7 +74,7 @@ public:
 
 private:
 	/* Returns true, when the node can/should not be optimized.*/
-	bool ShouldExpandInsteadCompile() const;
+	UE_API bool ShouldExpandInsteadCompile() const;
 
 	/**
 	* Clears this node's sub-graph, and then takes the supplied string and
@@ -80,20 +82,20 @@ private:
 	*
 	* @param  NewExpression	The new string we wish to replace the current "Expression".
 	*/
-	void RebuildExpression(FString NewExpression);
+	UE_API void RebuildExpression(FString NewExpression);
 
 	/**
 	* Clears the cached expression string, deletes all generated node, clears
 	* input pins, and resets the parser and graph generator (priming this node
 	* for a new expression).
 	*/
-	void ClearExpression();
+	UE_API void ClearExpression();
 
 	/** Sanitizes an expression for display, removing outermost parentheses */
-	FString SanitizeDisplayExpression(FString InExpression) const;
+	UE_API FString SanitizeDisplayExpression(FString InExpression) const;
 
 	/** Helper function to build the node's full title */
-	FText GetFullTitle(FText InExpression) const;
+	UE_API FText GetFullTitle(FText InExpression) const;
 private:
 	/** Cached so we don't have to regenerate it when the graph is recompiled */
 	TSharedPtr<class FCompilerResultsLog> CachedMessageLog;
@@ -106,3 +108,4 @@ private:
 };
 
 
+#undef UE_API
diff --git a/Engine/Source/Editor/BlueprintGraph/Classes/K2Node_SpawnActorFromClass.h b/Engine/Source/Editor/BlueprintGraph/Classes/K2Node_SpawnActorFromClass.h
--- a/Engine/Source/Editor/BlueprintGraph/Classes/K2Node_SpawnActorFromClass.h
+++ b/Engine/Source/Editor/BlueprintGraph/Classes/K2Node_SpawnActorFromClass.h
@@ -18,6 +18,8 @@
 
 #include "K2Node_SpawnActorFromClass.generated.h"
 
+#define UE_API BLUEPRINTGRAPH_API
+
 class FBlueprintActionDatabaseRegistrar;
 class FString;
 class UClass;
@@ -27,51 +29,53 @@ class UObject;
 struct FLinearColor;
 template <typename KeyType, typename ValueType> struct TKeyValuePair;
 
-UCLASS()
-class BLUEPRINTGRAPH_API UK2Node_SpawnActorFromClass : public UK2Node_ConstructObjectFromClass
+UCLASS(MinimalAPI)
+class UK2Node_SpawnActorFromClass : public UK2Node_ConstructObjectFromClass
 {
 	GENERATED_UCLASS_BODY()
 
 	//~ Begin UEdGraphNode Interface.
-	virtual void AllocateDefaultPins() override;
-	virtual FText GetNodeTitle(ENodeTitleType::Type TitleType) const override;
-	virtual void GetPinHoverText(const UEdGraphPin& Pin, FString& HoverTextOut) const override;
-	virtual void ExpandNode(class FKismetCompilerContext& CompilerContext, UEdGraph* SourceGraph) override;
-	virtual FSlateIcon GetIconAndTint(FLinearColor& OutColor) const override;
-	virtual bool IsCompatibleWithGraph(const UEdGraph* TargetGraph) const override;
-	virtual void PostPlacedNewNode() override;
+	UE_API virtual void AllocateDefaultPins() override;
+	UE_API virtual FText GetNodeTitle(ENodeTitleType::Type TitleType) const override;
+	UE_API virtual void GetPinHoverText(const UEdGraphPin& Pin, FString& HoverTextOut) const override;
+	UE_API virtual void ExpandNode(class FKismetCompilerContext& CompilerContext, UEdGraph* SourceGraph) override;
+	UE_API virtual FSlateIcon GetIconAndTint(FLinearColor& OutColor) const override;
+	UE_API virtual bool IsCompatibleWithGraph(const UEdGraph* TargetGraph) const override;
+	UE_API virtual void PostPlacedNewNode() override;
 	//~ End UEdGraphNode Interface.
 
 	//~ Begin UObject Interface
-	virtual void Serialize(FArchive& Ar) override;
-	virtual void PostLoad() override;
+	UE_API virtual void Serialize(FArchive& Ar) override;
+	UE_API virtual void PostLoad() override;
 	//~ End UObject Interface
 	
 	//~ Begin UK2Node Interface
 	virtual bool IsNodeSafeToIgnore() const override { return true; }
-	virtual void ReallocatePinsDuringReconstruction(TArray<UEdGraphPin*>& OldPins) override;
-	virtual void GetNodeAttributes( TArray<TKeyValuePair<FString, FString>>& OutNodeAttributes ) const override;
-	virtual class FNodeHandlingFunctor* CreateNodeHandler(class FKismetCompilerContext& CompilerContext) const override;
+	UE_API virtual void ReallocatePinsDuringReconstruction(TArray<UEdGraphPin*>& OldPins) override;
+	UE_API virtual void GetNodeAttributes( TArray<TKeyValuePair<FString, FString>>& OutNodeAttributes ) const override;
+	UE_API virtual class FNodeHandlingFunctor* CreateNodeHandler(class FKismetCompilerContext& CompilerContext) const override;
 	//~ End UK2Node Interface
 
 	//~ Begin UK2Node_ConstructObjectFromClass Interface
-	virtual UClass* GetClassPinBaseClass() const;
-	virtual bool IsSpawnVarPin(UEdGraphPin* Pin) const override;
+	UE_API virtual UClass* GetClassPinBaseClass() const;
+	UE_API virtual bool IsSpawnVarPin(UEdGraphPin* Pin) const override;
 	//~ End UK2Node_ConstructObjectFromClass Interface
 
 	
 private:
-	void FixupScaleMethodPin();
+	UE_API void FixupScaleMethodPin();
 	
 	/** Get the spawn transform input pin */	
-	UEdGraphPin* GetSpawnTransformPin() const;
+	UE_API UEdGraphPin* GetSpawnTransformPin() const;
 	/** Get the collision handling method input pin */
-	UEdGraphPin* GetCollisionHandlingOverridePin() const;
+	UE_API UEdGraphPin* GetCollisionHandlingOverridePin() const;
 	/** Get the collision handling method input pin */
-	UEdGraphPin* GetScaleMethodPin() const;
-	UEdGraphPin* TryGetScaleMethodPin() const;
+	UE_API UEdGraphPin* GetScaleMethodPin() const;
+	UE_API UEdGraphPin* TryGetScaleMethodPin() const;
 	/** Get the actor owner pin */
-	UEdGraphPin* GetOwnerPin() const;
+	UE_API UEdGraphPin* GetOwnerPin() const;
 
-	void MaybeUpdateCollisionPin(TArray<UEdGraphPin*>& OldPins);
+	UE_API void MaybeUpdateCollisionPin(TArray<UEdGraphPin*>& OldPins);
 };
+
+#undef UE_API
diff --git a/Engine/Source/Editor/BlueprintGraph/Classes/K2Node_Switch.h b/Engine/Source/Editor/BlueprintGraph/Classes/K2Node_Switch.h
--- a/Engine/Source/Editor/BlueprintGraph/Classes/K2Node_Switch.h
+++ b/Engine/Source/Editor/BlueprintGraph/Classes/K2Node_Switch.h
@@ -21,10 +21,12 @@
 
 #include "K2Node_Switch.generated.h"
 
+#define UE_API BLUEPRINTGRAPH_API
+
 class UObject;
 
-UCLASS(abstract)
-class BLUEPRINTGRAPH_API UK2Node_Switch : public UK2Node
+UCLASS(MinimalAPI, abstract)
+class UK2Node_Switch : public UK2Node
 {
 	GENERATED_UCLASS_BODY()
 
@@ -41,25 +43,25 @@ class BLUEPRINTGRAPH_API UK2Node_Switch : public UK2Node
 	TSubclassOf<class UObject> FunctionClass;
 
 	// UObject interface
-	virtual void PostEditChangeProperty(struct FPropertyChangedEvent& PropertyChangedEvent) override;
+	UE_API virtual void PostEditChangeProperty(struct FPropertyChangedEvent& PropertyChangedEvent) override;
 	// End of UObject interface
 
 	// UEdGraphNode interface
-	virtual void AllocateDefaultPins() override;
-	virtual FLinearColor GetNodeTitleColor() const override;
-	virtual FSlateIcon GetIconAndTint(FLinearColor& OutColor) const override;
+	UE_API virtual void AllocateDefaultPins() override;
+	UE_API virtual FLinearColor GetNodeTitleColor() const override;
+	UE_API virtual FSlateIcon GetIconAndTint(FLinearColor& OutColor) const override;
 	// End of UEdGraphNode interface
 
 	// UK2Node interface
-	virtual ERedirectType DoPinsMatchForReconstruction(const UEdGraphPin* NewPin, int32 NewPinIndex, const UEdGraphPin* OldPin, int32 OldPinIndex) const override;
-	virtual class FNodeHandlingFunctor* CreateNodeHandler(class FKismetCompilerContext& CompilerContext) const override;
-	virtual FText GetMenuCategory() const override;
+	UE_API virtual ERedirectType DoPinsMatchForReconstruction(const UEdGraphPin* NewPin, int32 NewPinIndex, const UEdGraphPin* OldPin, int32 OldPinIndex) const override;
+	UE_API virtual class FNodeHandlingFunctor* CreateNodeHandler(class FKismetCompilerContext& CompilerContext) const override;
+	UE_API virtual FText GetMenuCategory() const override;
 	virtual bool CanEverRemoveExecutionPin() const override { return true; }
 	// End of UK2Node interface
 
 	// UK2Node_Switch interface
 
-	virtual FString GetExportTextForPin(const UEdGraphPin* Pin) const;
+	UE_API virtual FString GetExportTextForPin(const UEdGraphPin* Pin) const;
 
 	/** Gets a unique pin name, the next in the sequence */
 	virtual FName GetUniquePinName() { return NAME_None; }
@@ -67,39 +69,39 @@ class BLUEPRINTGRAPH_API UK2Node_Switch : public UK2Node
 	/** Gets the pin type from the schema for the subclass */
 	virtual FEdGraphPinType GetPinType() const PURE_VIRTUAL(UK2Node_Switch::GetPinType, return FEdGraphPinType();)
 
-	virtual FEdGraphPinType GetInnerCaseType() const;
+	UE_API virtual FEdGraphPinType GetInnerCaseType() const;
 
 	/**
 	 * Adds a new execution pin to a switch node
 	 */
-	virtual void AddPinToSwitchNode();
+	UE_API virtual void AddPinToSwitchNode();
 
 	/**
 	 * Removes the specified execution pin from an switch node
 	 *
 	 * @param	TargetPin	The pin to remove from the node
 	 */
-	virtual void RemovePinFromSwitchNode(UEdGraphPin* TargetPin);
+	UE_API virtual void RemovePinFromSwitchNode(UEdGraphPin* TargetPin);
 
 	/** Whether an execution pin can be removed from the node or not */
-	virtual bool CanRemoveExecutionPin(UEdGraphPin* TargetPin) const;
+	UE_API virtual bool CanRemoveExecutionPin(UEdGraphPin* TargetPin) const;
 
 	/** Determines whether the add button should be enabled */
 	virtual bool SupportsAddPinButton() const { return true; }
 
 	/** Getting pin access */
-	UEdGraphPin* GetSelectionPin() const;
-	UEdGraphPin* GetDefaultPin() const;
-	UEdGraphPin* GetFunctionPin() const;
+	UE_API UEdGraphPin* GetSelectionPin() const;
+	UE_API UEdGraphPin* GetDefaultPin() const;
+	UE_API UEdGraphPin* GetFunctionPin() const;
 
-	static FName GetSelectionPinName();
+	static UE_API FName GetSelectionPinName();
 
-	virtual FName GetPinNameGivenIndex(int32 Index) const;
+	UE_API virtual FName GetPinNameGivenIndex(int32 Index) const;
 
 protected:
 	virtual void CreateSelectionPin() {}
 	virtual void CreateCasePins() {}
-	virtual void CreateFunctionPin();
+	UE_API virtual void CreateFunctionPin();
 	virtual void RemovePin(UEdGraphPin* TargetPin) {}
 
 private:
@@ -107,3 +109,4 @@ private:
 	bool bHasDefaultPinValueChanged;
 };
 
+#undef UE_API
diff --git a/Engine/Source/Editor/BlueprintGraph/Classes/K2Node_Tunnel.h b/Engine/Source/Editor/BlueprintGraph/Classes/K2Node_Tunnel.h
--- a/Engine/Source/Editor/BlueprintGraph/Classes/K2Node_Tunnel.h
+++ b/Engine/Source/Editor/BlueprintGraph/Classes/K2Node_Tunnel.h
@@ -19,6 +19,8 @@
 
 #include "K2Node_Tunnel.generated.h"
 
+#define UE_API BLUEPRINTGRAPH_API
+
 class FArchive;
 class FCompilerResultsLog;
 class UBlueprint;
@@ -27,8 +29,8 @@ class UEdGraphPin;
 class UObject;
 struct FEdGraphPinType;
 
-UCLASS()
-class BLUEPRINTGRAPH_API UK2Node_Tunnel : public UK2Node_EditablePinBase
+UCLASS(MinimalAPI)
+class UK2Node_Tunnel : public UK2Node_EditablePinBase
 {
 	GENERATED_UCLASS_BODY()
 
@@ -57,35 +59,35 @@ class BLUEPRINTGRAPH_API UK2Node_Tunnel : public UK2Node_EditablePinBase
 	struct FKismetUserDeclaredFunctionMetadata MetaData;
 
 	//~ Begin UEdGraphNode Interface.
-	virtual void DestroyNode() override;
-	virtual void PostPasteNode() override;
-	virtual FText GetTooltipText() const override;
-	virtual FText GetNodeTitle(ENodeTitleType::Type TitleType) const override;
-	virtual bool CanUserDeleteNode() const override;
-	virtual bool CanDuplicateNode() const override;
-	virtual UObject* GetJumpTargetForDoubleClick() const override;
-	virtual FName CreateUniquePinName(FName SourcePinName) const override;
-	virtual void ReallocatePinsDuringReconstruction(TArray<UEdGraphPin*>& OldPins) override;
-	virtual bool IsCompatibleWithGraph(const UEdGraph* InGraph) const override;
+	UE_API virtual void DestroyNode() override;
+	UE_API virtual void PostPasteNode() override;
+	UE_API virtual FText GetTooltipText() const override;
+	UE_API virtual FText GetNodeTitle(ENodeTitleType::Type TitleType) const override;
+	UE_API virtual bool CanUserDeleteNode() const override;
+	UE_API virtual bool CanDuplicateNode() const override;
+	UE_API virtual UObject* GetJumpTargetForDoubleClick() const override;
+	UE_API virtual FName CreateUniquePinName(FName SourcePinName) const override;
+	UE_API virtual void ReallocatePinsDuringReconstruction(TArray<UEdGraphPin*>& OldPins) override;
+	UE_API virtual bool IsCompatibleWithGraph(const UEdGraph* InGraph) const override;
 	//~ End UEdGraphNode Interface
 
 	//~ Begin UK2Node Interface.
-	virtual bool IsNodeSafeToIgnore() const override;
-	virtual bool DrawNodeAsEntry() const override;
-	virtual bool DrawNodeAsExit() const override;
+	UE_API virtual bool IsNodeSafeToIgnore() const override;
+	UE_API virtual bool DrawNodeAsEntry() const override;
+	UE_API virtual bool DrawNodeAsExit() const override;
 	virtual bool NodeCausesStructuralBlueprintChange() const override { return true; }
-	virtual void ClearCachedBlueprintData(UBlueprint* Blueprint) override;
-	virtual void ValidateNodeDuringCompilation(FCompilerResultsLog& MessageLog) const override;
-	virtual void FixupPinStringDataReferences(FArchive* SavingArchive) override;
+	UE_API virtual void ClearCachedBlueprintData(UBlueprint* Blueprint) override;
+	UE_API virtual void ValidateNodeDuringCompilation(FCompilerResultsLog& MessageLog) const override;
+	UE_API virtual void FixupPinStringDataReferences(FArchive* SavingArchive) override;
 	//~ End UK2Node Interface
 
 	//~ Begin UK2Node_EditablePinBase Interface.
-	virtual UEdGraphPin* CreatePinFromUserDefinition(const TSharedPtr<FUserPinInfo> NewPinInfo) override;
-	virtual bool CanModifyExecutionWires() override;
-	virtual ERenamePinResult RenameUserDefinedPinImpl(const FName OldName, const FName NewName, bool bTest = false) override;
+	UE_API virtual UEdGraphPin* CreatePinFromUserDefinition(const TSharedPtr<FUserPinInfo> NewPinInfo) override;
+	UE_API virtual bool CanModifyExecutionWires() override;
+	UE_API virtual ERenamePinResult RenameUserDefinedPinImpl(const FName OldName, const FName NewName, bool bTest = false) override;
 	virtual bool CanUseRefParams() const override { return true; }
-	virtual bool CanCreateUserDefinedPin(const FEdGraphPinType& InPinType, EEdGraphPinDirection InDesiredDirection, FText& OutErrorMessage) override;
-	virtual bool ModifyUserDefinedPinDefaultValue(TSharedPtr<FUserPinInfo> PinInfo, const FString& NewDefaultValue) override;
+	UE_API virtual bool CanCreateUserDefinedPin(const FEdGraphPinType& InPinType, EEdGraphPinDirection InDesiredDirection, FText& OutErrorMessage) override;
+	UE_API virtual bool ModifyUserDefinedPinDefaultValue(TSharedPtr<FUserPinInfo> PinInfo, const FString& NewDefaultValue) override;
 	//~ End UK2Node_EditablePinBase Interface
 
 protected:
@@ -97,21 +99,22 @@ protected:
 	virtual void PostFixupAllWildcardPins(bool bInAllWildcardPinsUnlinked) {}
 	// Feature flag for 'smart' wildcard inference which allows for multiple 
 	// types to be inferred.
-	static bool ShouldDoSmartWildcardInference();
+	static UE_API bool ShouldDoSmartWildcardInference();
 
 	// Utility function that subclasses must call after allocating their default pins
-	void CacheWildcardPins();
+	UE_API void CacheWildcardPins();
 
 	// Cache of the pins that were created as wildcards
 	TArray<UEdGraphPin*> WildcardPins;
 
 public:
 	// The input pins of this tunnel go to the output pins of InputSinkNode (can be NULL).
-	virtual UK2Node_Tunnel* GetInputSink() const;
+	UE_API virtual UK2Node_Tunnel* GetInputSink() const;
 
 	// The output pins of this tunnel node came from the input pins of OutputSourceNode (can be NULL).
-	virtual UK2Node_Tunnel* GetOutputSource() const;
+	UE_API virtual UK2Node_Tunnel* GetOutputSource() const;
 };
 
 
 
+#undef UE_API
diff --git a/Engine/Source/Editor/BlueprintGraph/Classes/K2Node_TunnelBoundary.h b/Engine/Source/Editor/BlueprintGraph/Classes/K2Node_TunnelBoundary.h
--- a/Engine/Source/Editor/BlueprintGraph/Classes/K2Node_TunnelBoundary.h
+++ b/Engine/Source/Editor/BlueprintGraph/Classes/K2Node_TunnelBoundary.h
@@ -14,6 +14,8 @@
 
 #include "K2Node_TunnelBoundary.generated.h"
 
+#define UE_API BLUEPRINTGRAPH_API
+
 class FKismetCompilerContext;
 class UK2Node_Tunnel;
 class UObject;
@@ -27,8 +29,8 @@ enum class ETunnelBoundaryType : uint8
 	OutputSite
 };
 
-UCLASS()
-class BLUEPRINTGRAPH_API UK2Node_TunnelBoundary : public UK2Node
+UCLASS(MinimalAPI)
+class UK2Node_TunnelBoundary : public UK2Node
 {
 	GENERATED_UCLASS_BODY()
 
@@ -43,11 +45,11 @@ class BLUEPRINTGRAPH_API UK2Node_TunnelBoundary : public UK2Node
 public:
 
 	//~ Begin UEdGraphNode Interface
-	virtual FText GetNodeTitle(ENodeTitleType::Type TitleType) const override;
+	UE_API virtual FText GetNodeTitle(ENodeTitleType::Type TitleType) const override;
 	//~ End UEdGraphNode Interface
 
 	//~ Begin K2Node Interface
-	virtual class FNodeHandlingFunctor* CreateNodeHandler(FKismetCompilerContext& CompilerContext) const override;
+	UE_API virtual class FNodeHandlingFunctor* CreateNodeHandler(FKismetCompilerContext& CompilerContext) const override;
 	virtual int32 GetNodeRefreshPriority() const override { return EBaseNodeRefreshPriority::Low_UsesDependentWildcard; }
 	//~ End UK2Node Interface
 
@@ -55,5 +57,7 @@ public:
 	ETunnelBoundaryType GetTunnelBoundaryType() const { return TunnelBoundaryType; }
 
 	/** Set node attributes based on the tunnel source node */
-	void SetNodeAttributes(const UK2Node_Tunnel* SourceNode);
-};
\ No newline at end of file
+	UE_API void SetNodeAttributes(const UK2Node_Tunnel* SourceNode);
+};
+
+#undef UE_API
diff --git a/Engine/Source/Editor/BlueprintGraph/Classes/K2Node_Variable.h b/Engine/Source/Editor/BlueprintGraph/Classes/K2Node_Variable.h
--- a/Engine/Source/Editor/BlueprintGraph/Classes/K2Node_Variable.h
+++ b/Engine/Source/Editor/BlueprintGraph/Classes/K2Node_Variable.h
@@ -23,6 +23,8 @@
 
 #include "K2Node_Variable.generated.h"
 
+#define UE_API BLUEPRINTGRAPH_API
+
 class FArchive;
 class FProperty;
 class UActorComponent;
@@ -46,8 +48,8 @@ namespace ESelfContextInfo
 	};
 }
 
-UCLASS(abstract)
-class BLUEPRINTGRAPH_API UK2Node_Variable : public UK2Node
+UCLASS(MinimalAPI, abstract)
+class UK2Node_Variable : public UK2Node
 {
 	GENERATED_UCLASS_BODY()
 
@@ -75,47 +77,47 @@ protected:
 	 * Remap a reference from one variable to another, if this variable is of class type 'MatchInVariableClass', and if linked to anything that is a child of 'RemapIfLinkedToClass'.
 	 * Only intended for versioned fixup where redirects can't be applied.
 	 */
-	bool RemapRestrictedLinkReference(FName OldVariableName, FName NewVariableName, const UClass* MatchInVariableClass, const UClass* RemapIfLinkedToClass, bool bLogWarning);
+	UE_API bool RemapRestrictedLinkReference(FName OldVariableName, FName NewVariableName, const UClass* MatchInVariableClass, const UClass* RemapIfLinkedToClass, bool bLogWarning);
 
 public:
 	//~ Begin UObject Interface
-	virtual void Serialize(FArchive& Ar) override;
+	UE_API virtual void Serialize(FArchive& Ar) override;
 	//~ End UObject Interface
 
 	//~ Begin UEdGraphNode Interface
-	virtual FLinearColor GetNodeTitleColor() const override;
-	virtual FString GetFindReferenceSearchString_Impl(EGetFindReferenceSearchStringFlags InFlags) const override;
-	virtual void ReconstructNode() override;
-	virtual FString GetDocumentationLink() const override;
-	virtual FString GetDocumentationExcerptName() const override;
-	virtual FName GetCornerIcon() const override;
-	virtual FSlateIcon GetIconAndTint(FLinearColor& OutColor) const override;
-	virtual void AutowireNewNode(UEdGraphPin* FromPin) override;
-	virtual bool CanPasteHere(const UEdGraph* TargetGraph) const override;
-	virtual void PostPasteNode() override;
-	virtual bool HasDeprecatedReference() const override;
-	virtual FEdGraphNodeDeprecationResponse GetDeprecationResponse(EEdGraphNodeDeprecationType DeprecationType) const override;
-	virtual UObject* GetJumpTargetForDoubleClick() const override;
-	virtual bool CanJumpToDefinition() const override;
-	virtual void JumpToDefinition() const override;
-	virtual void GetNodeContextMenuActions(class UToolMenu* Menu, class UGraphNodeContextMenuContext* Context) const override;
-	virtual FString GetPinMetaData(FName InPinName, FName InKey) override;
-	virtual bool HasExternalDependencies(TArray<class UStruct*>* OptionalOutput) const override;
+	UE_API virtual FLinearColor GetNodeTitleColor() const override;
+	UE_API virtual FString GetFindReferenceSearchString_Impl(EGetFindReferenceSearchStringFlags InFlags) const override;
+	UE_API virtual void ReconstructNode() override;
+	UE_API virtual FString GetDocumentationLink() const override;
+	UE_API virtual FString GetDocumentationExcerptName() const override;
+	UE_API virtual FName GetCornerIcon() const override;
+	UE_API virtual FSlateIcon GetIconAndTint(FLinearColor& OutColor) const override;
+	UE_API virtual void AutowireNewNode(UEdGraphPin* FromPin) override;
+	UE_API virtual bool CanPasteHere(const UEdGraph* TargetGraph) const override;
+	UE_API virtual void PostPasteNode() override;
+	UE_API virtual bool HasDeprecatedReference() const override;
+	UE_API virtual FEdGraphNodeDeprecationResponse GetDeprecationResponse(EEdGraphNodeDeprecationType DeprecationType) const override;
+	UE_API virtual UObject* GetJumpTargetForDoubleClick() const override;
+	UE_API virtual bool CanJumpToDefinition() const override;
+	UE_API virtual void JumpToDefinition() const override;
+	UE_API virtual void GetNodeContextMenuActions(class UToolMenu* Menu, class UGraphNodeContextMenuContext* Context) const override;
+	UE_API virtual FString GetPinMetaData(FName InPinName, FName InKey) override;
+	UE_API virtual bool HasExternalDependencies(TArray<class UStruct*>* OptionalOutput) const override;
 	//~ End UEdGraphNode Interface
 
 	//~ Begin K2Node Interface
 	virtual bool DrawNodeAsVariable() const override { return true; }
-	virtual ERedirectType DoPinsMatchForReconstruction(const UEdGraphPin* NewPin, int32 NewPinIndex, const UEdGraphPin* OldPin, int32 OldPinIndex)  const override;
-	virtual void ValidateNodeDuringCompilation(class FCompilerResultsLog& MessageLog) const override;
-	virtual FText GetToolTipHeading() const override;
-	virtual void GetNodeAttributes(TArray<TKeyValuePair<FString, FString>>& OutNodeAttributes) const override;
-	virtual void HandleVariableRenamed(UBlueprint* InBlueprint, UClass* InVariableClass, UEdGraph* InGraph, const FName& InOldVarName, const FName& InNewVarName) override;
-	virtual void ReplaceReferences(UBlueprint* InBlueprint, UBlueprint* InReplacementBlueprint, const FMemberReference& InSource, const FMemberReference& InReplacement) override;
-	virtual bool ReferencesVariable(const FName& InVarName, const UStruct* InScope) const override;
+	UE_API virtual ERedirectType DoPinsMatchForReconstruction(const UEdGraphPin* NewPin, int32 NewPinIndex, const UEdGraphPin* OldPin, int32 OldPinIndex)  const override;
+	UE_API virtual void ValidateNodeDuringCompilation(class FCompilerResultsLog& MessageLog) const override;
+	UE_API virtual FText GetToolTipHeading() const override;
+	UE_API virtual void GetNodeAttributes(TArray<TKeyValuePair<FString, FString>>& OutNodeAttributes) const override;
+	UE_API virtual void HandleVariableRenamed(UBlueprint* InBlueprint, UClass* InVariableClass, UEdGraph* InGraph, const FName& InOldVarName, const FName& InNewVarName) override;
+	UE_API virtual void ReplaceReferences(UBlueprint* InBlueprint, UBlueprint* InReplacementBlueprint, const FMemberReference& InSource, const FMemberReference& InReplacement) override;
+	UE_API virtual bool ReferencesVariable(const FName& InVarName, const UStruct* InScope) const override;
 	//~ End K2Node Interface
 
 	/** Set up this variable node from the supplied FProperty */
-	void SetFromProperty(const FProperty* Property, bool bSelfContext, UClass* OwnerClass);
+	UE_API void SetFromProperty(const FProperty* Property, bool bSelfContext, UClass* OwnerClass);
 
 	/** Util to get variable name as a string */
 	FString GetVarNameString() const
@@ -143,10 +145,10 @@ public:
 	 *
 	 * @return	true if it succeeds, false if it fails.
 	 */
-	bool CreatePinForVariable(EEdGraphPinDirection Direction, FName InPinName = NAME_None);
+	UE_API bool CreatePinForVariable(EEdGraphPinDirection Direction, FName InPinName = NAME_None);
 
 	/** Creates 'self' pin */
-	void CreatePinForSelf();
+	UE_API void CreatePinForSelf();
 
 	/**
 	 * Creates a reader or writer pin for a variable from an old pin.
@@ -157,31 +159,31 @@ public:
 	 *
 	 * @return	true if it succeeds, false if it fails.
 	 */
-	bool RecreatePinForVariable(EEdGraphPinDirection Direction, TArray<UEdGraphPin*>& OldPins, FName InPinName = NAME_None);
+	UE_API bool RecreatePinForVariable(EEdGraphPinDirection Direction, TArray<UEdGraphPin*>& OldPins, FName InPinName = NAME_None);
 
 	/** Get the class to look for this variable in */
-	UClass* GetVariableSourceClass() const;
+	UE_API UClass* GetVariableSourceClass() const;
 
 	/** Get the FProperty for this variable node */
-	FProperty* GetPropertyForVariable() const;
-	FProperty* GetPropertyForVariableFromSkeleton() const;
+	UE_API FProperty* GetPropertyForVariable() const;
+	UE_API FProperty* GetPropertyForVariableFromSkeleton() const;
 
 	/** Returns true if the variable names match, this looks for redirectors */
-	static bool DoesRenamedVariableMatch(FName OldVariableName, FName NewVariableName, UStruct* StructType);
+	static UE_API bool DoesRenamedVariableMatch(FName OldVariableName, FName NewVariableName, UStruct* StructType);
 
 private:
 	/** 
 	 * Gets the property for the variable on the owning class or on the owning class's sparse class data structure.
 	 */
-	FProperty* GetPropertyForVariable_Internal(UClass* OwningClass) const;
+	UE_API FProperty* GetPropertyForVariable_Internal(UClass* OwningClass) const;
 
 public:
 
 	/** Accessor for the value output pin of the node */
-	UEdGraphPin* GetValuePin() const;
+	UE_API UEdGraphPin* GetValuePin() const;
 
 	/** Validates there are no errors in the node */
-	void CheckForErrors(const class UEdGraphSchema_K2* Schema, class FCompilerResultsLog& MessageLog);
+	UE_API void CheckForErrors(const class UEdGraphSchema_K2* Schema, class FCompilerResultsLog& MessageLog);
 
 	/**
 	 * Utility method intended to serve as a choke point for various slate 
@@ -192,7 +194,7 @@ public:
 	 * @param  IconColorOut	A color out, further discerning the variable's type.
 	 * @return A icon representing the specified variable's type.
 	 */
-	static FSlateIcon GetVariableIconAndColor(const UStruct* VarScope, FName VarName, FLinearColor& IconColorOut);
+	static UE_API FSlateIcon GetVariableIconAndColor(const UStruct* VarScope, FName VarName, FLinearColor& IconColorOut);
 
 	/**
 	 * Utility method intended to serve as a choke point for various slate 
@@ -202,7 +204,7 @@ public:
 	 * @param  IconColorOut	A color out, further discerning the variable's type.
 	 * @return A icon representing the specified variable's type.
 	 */
-	static FSlateIcon GetVarIconFromPinType(const FEdGraphPinType& InPinType, FLinearColor& IconColorOut);
+	static UE_API FSlateIcon GetVarIconFromPinType(const FEdGraphPinType& InPinType, FLinearColor& IconColorOut);
 
 protected:
 	/**
@@ -210,19 +212,20 @@ protected:
 	 * 
 	 * @return 
 	 */
-	FBPVariableDescription const* GetBlueprintVarDescription() const;
+	UE_API FBPVariableDescription const* GetBlueprintVarDescription() const;
 
 	/**
 	 * Utility function to retrieve actor component for variable property.
 	 *
 	 * @return Found UActorComponent for the FProperty, else nullptr
 	 */
-	const UActorComponent* GetActorComponent(const FProperty* VariableProperty) const;
+	UE_API const UActorComponent* GetActorComponent(const FProperty* VariableProperty) const;
 
 	/** Adds the variable reference to the suppressed deprecation warnings list */
-	void SuppressDeprecationWarning() const;
+	UE_API void SuppressDeprecationWarning() const;
 
 	/** Returns whether a Function Graph contains a parameter with the given name */
-	static bool FunctionParameterExists(const UEdGraph* InFunctionGraph, const FName InParameterName);
+	static UE_API bool FunctionParameterExists(const UEdGraph* InFunctionGraph, const FName InParameterName);
 };
 
+#undef UE_API
diff --git a/Engine/Source/Editor/BlueprintGraph/Classes/K2Node_VariableGet.h b/Engine/Source/Editor/BlueprintGraph/Classes/K2Node_VariableGet.h
--- a/Engine/Source/Editor/BlueprintGraph/Classes/K2Node_VariableGet.h
+++ b/Engine/Source/Editor/BlueprintGraph/Classes/K2Node_VariableGet.h
@@ -17,6 +17,8 @@
 
 #include "K2Node_VariableGet.generated.h"
 
+#define UE_API BLUEPRINTGRAPH_API
+
 class FArchive;
 class FProperty;
 class UEdGraph;
@@ -33,44 +35,44 @@ enum class EGetNodeVariation
 	Branch,
 };
 
-UCLASS()
-class BLUEPRINTGRAPH_API UK2Node_VariableGet : public UK2Node_Variable
+UCLASS(MinimalAPI)
+class UK2Node_VariableGet : public UK2Node_Variable
 {
 	GENERATED_UCLASS_BODY()
 
 public:
 
 	//~ Begin UObject Interface
-	virtual void Serialize(FArchive& Ar) override;
+	UE_API virtual void Serialize(FArchive& Ar) override;
 	//~ End UObject Interface
 
 	//~ Begin UEdGraphNode Interface
-	virtual void AllocateDefaultPins() override;
-	virtual FText GetTooltipText() const override;
-	virtual FText GetNodeTitle(ENodeTitleType::Type TitleType) const override;
-	virtual void GetNodeContextMenuActions(class UToolMenu* Menu, class UGraphNodeContextMenuContext* Context) const override;
+	UE_API virtual void AllocateDefaultPins() override;
+	UE_API virtual FText GetTooltipText() const override;
+	UE_API virtual FText GetNodeTitle(ENodeTitleType::Type TitleType) const override;
+	UE_API virtual void GetNodeContextMenuActions(class UToolMenu* Menu, class UGraphNodeContextMenuContext* Context) const override;
 	virtual bool IncludeParentNodeContextMenu() const override { return true; }
-	virtual void ValidateNodeDuringCompilation(class FCompilerResultsLog& MessageLog) const override;
+	UE_API virtual void ValidateNodeDuringCompilation(class FCompilerResultsLog& MessageLog) const override;
 	//~ End UEdGraphNode Interface
 
 	//~ Begin K2Node Interface
 	virtual bool IsNodePure() const override { return CurrentVariation == EGetNodeVariation::Pure; }
-	virtual void ReallocatePinsDuringReconstruction(TArray<UEdGraphPin*>& OldPins) override;
+	UE_API virtual void ReallocatePinsDuringReconstruction(TArray<UEdGraphPin*>& OldPins) override;
 	virtual bool ShouldShowNodeProperties() const override { return true; }
-	virtual class FNodeHandlingFunctor* CreateNodeHandler(class FKismetCompilerContext& CompilerContext) const override;
-	virtual void ExpandNode(class FKismetCompilerContext& CompilerContext, UEdGraph* SourceGraph) override;
+	UE_API virtual class FNodeHandlingFunctor* CreateNodeHandler(class FKismetCompilerContext& CompilerContext) const override;
+	UE_API virtual void ExpandNode(class FKismetCompilerContext& CompilerContext, UEdGraph* SourceGraph) override;
 	//~ End K2Node Interface
 
-	static FText GetPropertyTooltip(const FProperty* VariableProperty);
-	static FText GetBlueprintVarTooltip(const FBPVariableDescription& VarDesc);
+	static UE_API FText GetPropertyTooltip(const FProperty* VariableProperty);
+	static UE_API FText GetBlueprintVarTooltip(const FBPVariableDescription& VarDesc);
 
 private:
 
 	/** Adds pins required for the node to function in a impure manner */
-	void CreateImpurePins(TArray<UEdGraphPin*>* InOldPinsPtr);
+	UE_API void CreateImpurePins(TArray<UEdGraphPin*>* InOldPinsPtr);
 
 	/** Flips the node's purity (adding/removing exec pins as needed). */
-	void TogglePurity(EGetNodeVariation BoundVariation);
+	UE_API void TogglePurity(EGetNodeVariation BoundVariation);
 
 	/** Constructing FText strings can be costly, so we cache the node's title/tooltip */
 	FNodeTextCache CachedTooltip;
@@ -83,3 +85,4 @@ private:
 	EGetNodeVariation CurrentVariation;
 };
 
+#undef UE_API
diff --git a/Engine/Source/Editor/BlueprintGraph/Classes/NodeDependingOnEnumInterface.h b/Engine/Source/Editor/BlueprintGraph/Classes/NodeDependingOnEnumInterface.h
--- a/Engine/Source/Editor/BlueprintGraph/Classes/NodeDependingOnEnumInterface.h
+++ b/Engine/Source/Editor/BlueprintGraph/Classes/NodeDependingOnEnumInterface.h
@@ -13,13 +13,13 @@
 
 class UObject;
 
-UINTERFACE(meta=(CannotImplementInterfaceInBlueprint))
-class BLUEPRINTGRAPH_API UNodeDependingOnEnumInterface : public UInterface
+UINTERFACE(MinimalAPI, meta=(CannotImplementInterfaceInBlueprint))
+class UNodeDependingOnEnumInterface : public UInterface
 {
 	GENERATED_UINTERFACE_BODY()
 };
 
-class BLUEPRINTGRAPH_API INodeDependingOnEnumInterface
+class INodeDependingOnEnumInterface
 {
 	GENERATED_IINTERFACE_BODY()
 
diff --git a/Engine/Source/Editor/BlueprintGraph/Public/AnimNotifyEventNodeSpawner.h b/Engine/Source/Editor/BlueprintGraph/Public/AnimNotifyEventNodeSpawner.h
--- a/Engine/Source/Editor/BlueprintGraph/Public/AnimNotifyEventNodeSpawner.h
+++ b/Engine/Source/Editor/BlueprintGraph/Public/AnimNotifyEventNodeSpawner.h
@@ -6,8 +6,10 @@
 
 #include "AnimNotifyEventNodeSpawner.generated.h"
 
-UCLASS(Transient)
-class BLUEPRINTGRAPH_API UAnimNotifyEventNodeSpawner : public UBlueprintEventNodeSpawner
+#define UE_API BLUEPRINTGRAPH_API
+
+UCLASS(MinimalAPI, Transient)
+class UAnimNotifyEventNodeSpawner : public UBlueprintEventNodeSpawner
 {
 	GENERATED_BODY()
 
@@ -18,7 +20,7 @@ public:
 	 * @param  NotifyName	The name you want assigned to the event.
 	 * @return A newly allocated instance of this class.
 	 */
-	static UAnimNotifyEventNodeSpawner* Create(const FSoftObjectPath& InSkeletonObjectPath, FName InNotifyName);
+	static UE_API UAnimNotifyEventNodeSpawner* Create(const FSoftObjectPath& InSkeletonObjectPath, FName InNotifyName);
 
 	/** @return the skeleton object path */
 	const FSoftObjectPath& GetSkeletonObjectPath() const { return SkeletonObjectPath; }
@@ -28,3 +30,5 @@ private:
 	UPROPERTY()
 	FSoftObjectPath SkeletonObjectPath;
 };
+
+#undef UE_API
diff --git a/Engine/Source/Editor/BlueprintGraph/Public/BlueprintActionDatabase.h b/Engine/Source/Editor/BlueprintGraph/Public/BlueprintActionDatabase.h
--- a/Engine/Source/Editor/BlueprintGraph/Public/BlueprintActionDatabase.h
+++ b/Engine/Source/Editor/BlueprintGraph/Public/BlueprintActionDatabase.h
@@ -20,6 +20,8 @@
 #include "UObject/SoftObjectPath.h"
 #include "UObject/TopLevelAssetPath.h"
 
+#define UE_API BLUEPRINTGRAPH_API
+
 class FBlueprintActionDatabaseRegistrar;
 class FReferenceCollector;
 class UBlueprint;
@@ -41,7 +43,7 @@ struct FEdGraphPinType;
  * 
  * @TODO:  Hook up to handle class recompile events, along with enum/struct asset creation events.
  */
-class BLUEPRINTGRAPH_API FBlueprintActionDatabase : public FGCObject, public FTickableEditorObject
+class FBlueprintActionDatabase : public FGCObject, public FTickableEditorObject
 {
 public:
 	DECLARE_MULTICAST_DELEGATE_OneParam(FOnDatabaseEntryUpdated, UObject*);
@@ -53,7 +55,7 @@ public:
 
 public:
 	/** Destructor */
-	virtual ~FBlueprintActionDatabase();
+	UE_API virtual ~FBlueprintActionDatabase();
 
 	/**
 	 * Getter to access the database singleton. Will populate the database first 
@@ -61,20 +63,20 @@ public:
 	 *
 	 * @return The singleton instance of FBlueprintActionDatabase.
 	 */
-	static FBlueprintActionDatabase& Get();
+	static UE_API FBlueprintActionDatabase& Get();
 
 	/** Getter to access the datbase singleton, will return null if the database has not been initialized */
-	static FBlueprintActionDatabase* TryGet();
+	static UE_API FBlueprintActionDatabase* TryGet();
 
 	// FTickableEditorObject interface
-	virtual void Tick(float DeltaTime) override;
+	UE_API virtual void Tick(float DeltaTime) override;
 	virtual ETickableTickType GetTickableTickType() const override { return ETickableTickType::Always; }
-	virtual TStatId GetStatId() const override;
+	UE_API virtual TStatId GetStatId() const override;
 	// End FTickableEditorObject interface
 
 	// FGCObject interface
-	virtual void AddReferencedObjects(FReferenceCollector& Collector) override;
-	virtual FString GetReferencerName() const override;
+	UE_API virtual void AddReferencedObjects(FReferenceCollector& Collector) override;
+	UE_API virtual FString GetReferencerName() const override;
 	// End FGCObject interface
 
 	/**
@@ -89,23 +91,23 @@ public:
 	 *
 	 * @return The entire action database, which maps specific objects to arrays of associated node-spawners.
 	 */
-	FActionRegistry const& GetAllActions();
+	UE_API FActionRegistry const& GetAllActions();
 
 	/**
 	 * Populates the action database from scratch. Loops over every known class
 	 * and records a set of node-spawners associated with each.
 	 */
-	void RefreshAll();
+	UE_API void RefreshAll();
 
 	/**
 	 * Populates the action database with all level script actions from all active editor worlds.
 	 */
-	void RefreshWorlds();
+	UE_API void RefreshWorlds();
 
 	/**
 	 * Removes the entry with the given key on the next tick.
 	 */
-	void DeferredRemoveEntry(FObjectKey const& InKey);
+	UE_API void DeferredRemoveEntry(FObjectKey const& InKey);
 
 	/**
 	 * Finds the database entry for the specified class and wipes it, 
@@ -113,7 +115,7 @@ public:
 	 *
 	 * @param  Class	The class entry you want rebuilt.
 	 */
-	void RefreshClassActions(UClass* const Class);
+	UE_API void RefreshClassActions(UClass* const Class);
 
 	/**
 	 * Finds the database entry for the specified asset and wipes it,
@@ -121,12 +123,12 @@ public:
 	 * 
 	 * @param  AssetObject	The asset entry you want rebuilt.
 	 */
-	void RefreshAssetActions(UObject* const AssetObject);
+	UE_API void RefreshAssetActions(UObject* const AssetObject);
 
 	/**
 	 * Updates all component related actions
 	 */
-	void RefreshComponentActions();
+	UE_API void RefreshComponentActions();
 
 	/**
 	 * Finds the database entry for the specified object and wipes it. The entry 
@@ -135,7 +137,7 @@ public:
 	 * @param  AssetObject	
 	 * @return True if an entry was found and removed.
 	 */
-	bool ClearAssetActions(UObject* const AssetObject);
+	UE_API bool ClearAssetActions(UObject* const AssetObject);
 
 	/**
 	 * Finds the database entry for the specified object and wipes it. The entry
@@ -144,7 +146,7 @@ public:
 	 * @param  AssetObjectKey
 	 * @return True if an entry was found and removed.
 	 */
-	bool ClearAssetActions(const FObjectKey& AssetObjectKey);
+	UE_API bool ClearAssetActions(const FObjectKey& AssetObjectKey);
 	
 	/**
 	 * Finds the database entry for the specified unloaded asset and wipes it.
@@ -152,7 +154,7 @@ public:
 	 *
 	 * @param ObjectPath	Object's path to lookup into the database
 	 */
-	void ClearUnloadedAssetActions(const FSoftObjectPath& ObjectPath);
+	UE_API void ClearUnloadedAssetActions(const FSoftObjectPath& ObjectPath);
 
 	/**
 	 * Moves the unloaded asset actions from one location to another
@@ -160,7 +162,7 @@ public:
 	 * @param SourceObjectPath	The object path that the data can currently be found under
 	 * @param TargetObjectPath	The object path that the data should be moved to
 	 */
-	void MoveUnloadedAssetActions(const FSoftObjectPath& SourceObjectPath, const FSoftObjectPath& TargetObjectPath);
+	UE_API void MoveUnloadedAssetActions(const FSoftObjectPath& SourceObjectPath, const FSoftObjectPath& TargetObjectPath);
 
 	/** */
 	FOnDatabaseEntryUpdated& OnEntryUpdated() { return EntryRefreshDelegate; }
@@ -187,48 +189,48 @@ public:
 	};
 	
 	/** Check whether the global filter applies to this class */
-	static bool IsClassAllowed(UClass const* InClass, EPermissionsContext InContext);
-	static bool IsClassAllowed(const FTopLevelAssetPath& InClassPath, EPermissionsContext InContext);
+	static UE_API bool IsClassAllowed(UClass const* InClass, EPermissionsContext InContext);
+	static UE_API bool IsClassAllowed(const FTopLevelAssetPath& InClassPath, EPermissionsContext InContext);
 
 	/** Check whether we have any class filtering applied */
-	static bool HasClassFiltering();
+	static UE_API bool HasClassFiltering();
 
 	/** Check whether permissions allow this field */
-	static bool IsFieldAllowed(UField const* InField, EPermissionsContext InContext);
+	static UE_API bool IsFieldAllowed(UField const* InField, EPermissionsContext InContext);
 
 	/** Check whether permissions allow this function */
-	static bool IsFunctionAllowed(UFunction const* InFunction, EPermissionsContext InContext);
+	static UE_API bool IsFunctionAllowed(UFunction const* InFunction, EPermissionsContext InContext);
 
 	/** Check whether permissions allow this enum */
-	static bool IsEnumAllowed(UEnum const* InEnum, EPermissionsContext InContext);
-	static bool IsEnumAllowed(const FTopLevelAssetPath& InEnumPath, EPermissionsContext InContext);
+	static UE_API bool IsEnumAllowed(UEnum const* InEnum, EPermissionsContext InContext);
+	static UE_API bool IsEnumAllowed(const FTopLevelAssetPath& InEnumPath, EPermissionsContext InContext);
 
 	/** Check whether permissions allow this struct */
-	static bool IsStructAllowed(UScriptStruct const* InStruct, EPermissionsContext InContext);
-	static bool IsStructAllowed(const FTopLevelAssetPath& InStructPath, EPermissionsContext InContext);
+	static UE_API bool IsStructAllowed(UScriptStruct const* InStruct, EPermissionsContext InContext);
+	static UE_API bool IsStructAllowed(const FTopLevelAssetPath& InStructPath, EPermissionsContext InContext);
 	
 	/** Check whether permissions allow this pin type. InUnloadedAssetPath will be used in the case where a InPinType.PinSubCategoryObject is unable to be resolved. */
-	static bool IsPinTypeAllowed(const FEdGraphPinType& InPinType, const FTopLevelAssetPath& InUnloadedAssetPath = FTopLevelAssetPath());
+	static UE_API bool IsPinTypeAllowed(const FEdGraphPinType& InPinType, const FTopLevelAssetPath& InUnloadedAssetPath = FTopLevelAssetPath());
 	
 private:
 	/** Private constructor for singleton purposes. */
-	FBlueprintActionDatabase();
-	void Init();
+	UE_API FBlueprintActionDatabase();
+	UE_API void Init();
 
 	/**
 	 * 
 	 * 
 	 * @param  Registrar	
 	 */
-	void RegisterAllNodeActions(FBlueprintActionDatabaseRegistrar& Registrar);
+	UE_API void RegisterAllNodeActions(FBlueprintActionDatabaseRegistrar& Registrar);
 
 	/**
 	 * This exists only because we need a pointer to associate our delegates with
 	 */
-	void OnBlueprintChanged( UBlueprint* );
+	UE_API void OnBlueprintChanged( UBlueprint* );
 
 	/** Refresh other systems before a full or partial refresh */
-	void PreRefresh(bool bRefreshAll);
+	UE_API void PreRefresh(bool bRefreshAll);
 private:
 	/** 
 	 * A map of associated node-spawners for each class/asset. A spawner that 
@@ -276,3 +278,5 @@ private:
 	/** Pointer to the shared list of currently existing component types */
 	const TArray<struct FComponentTypeEntry>* ComponentTypes;
 };
+
+#undef UE_API
diff --git a/Engine/Source/Editor/BlueprintGraph/Public/BlueprintActionDatabaseRegistrar.h b/Engine/Source/Editor/BlueprintGraph/Public/BlueprintActionDatabaseRegistrar.h
--- a/Engine/Source/Editor/BlueprintGraph/Public/BlueprintActionDatabaseRegistrar.h
+++ b/Engine/Source/Editor/BlueprintGraph/Public/BlueprintActionDatabaseRegistrar.h
@@ -12,6 +12,8 @@
 #include "UObject/Class.h"
 #include "UObject/Object.h"
 
+#define UE_API BLUEPRINTGRAPH_API
+
 class IAssetTools;
 class UBlueprintNodeSpawner;
 class UClass;
@@ -28,7 +30,7 @@ struct FAssetData;
  * a chance to register its own actions (specifically for UK2Nodes in other
  * modules that the database doesn't have access to).
  */
-class BLUEPRINTGRAPH_API FBlueprintActionDatabaseRegistrar
+class FBlueprintActionDatabaseRegistrar
 {
 public:
 	/**
@@ -38,7 +40,7 @@ public:
 	 *
 	 * @param  NodeSpawner	The new node spawner that you wish to register.
 	 */
-	bool AddBlueprintAction(UBlueprintNodeSpawner* NodeSpawner);
+	UE_API bool AddBlueprintAction(UBlueprintNodeSpawner* NodeSpawner);
 
 	/**
 	 * Each action should be recorded under a specific UField key; primarily to
@@ -48,11 +50,11 @@ public:
 	 * @param  FieldOwner	A field object that the node is associated with (flags when this action should be updated).
 	 * @param  NodeSpawner	The new node spawner that you wish to register.
 	 */
-	bool AddBlueprintAction(UClass const* ClassOwner, UBlueprintNodeSpawner* NodeSpawner);
-	bool AddBlueprintAction(UEnum const* EnumOwner, UBlueprintNodeSpawner* NodeSpawner);
-	bool AddBlueprintAction(UScriptStruct const* StructOwner, UBlueprintNodeSpawner* NodeSpawner);
-	bool AddBlueprintAction(UField const* FieldOwner, UBlueprintNodeSpawner* NodeSpawner);
-	bool AddBlueprintAction(FAssetData const& AssetDataOwner, UBlueprintNodeSpawner* NodeSpawner);
+	UE_API bool AddBlueprintAction(UClass const* ClassOwner, UBlueprintNodeSpawner* NodeSpawner);
+	UE_API bool AddBlueprintAction(UEnum const* EnumOwner, UBlueprintNodeSpawner* NodeSpawner);
+	UE_API bool AddBlueprintAction(UScriptStruct const* StructOwner, UBlueprintNodeSpawner* NodeSpawner);
+	UE_API bool AddBlueprintAction(UField const* FieldOwner, UBlueprintNodeSpawner* NodeSpawner);
+	UE_API bool AddBlueprintAction(FAssetData const& AssetDataOwner, UBlueprintNodeSpawner* NodeSpawner);
 
 	/**
 	 * Special case for asset bound actions (we want to clean-up/refresh these
@@ -61,7 +63,7 @@ public:
 	 * @param  AssetOwner	An asset object that the node is associated with (flags when this action should be updated).
 	 * @param  NodeSpawner	The new node spawner that you wish to register.
 	 */
-	bool AddBlueprintAction(UObject const* AssetOwner, UBlueprintNodeSpawner* NodeSpawner);
+	UE_API bool AddBlueprintAction(UObject const* AssetOwner, UBlueprintNodeSpawner* NodeSpawner);
 
 	/**
 	 * Occasionally (when an asset is added/refreshed), this registrar will be
@@ -73,8 +75,8 @@ public:
 	 * @param  OwnerKey		The key you wish to register your action(s) under.
 	 * @return True if the OwnerKey would is allowed to register actions, false if it would be blocked.
 	 */
-	bool IsOpenForRegistration(UObject const* OwnerKey);
-	bool IsOpenForRegistration(FAssetData const& AssetDataOwner);
+	UE_API bool IsOpenForRegistration(UObject const* OwnerKey);
+	UE_API bool IsOpenForRegistration(FAssetData const& AssetDataOwner);
 
 	/** Returns the current key the registrar is being filtered using */
 	UObject const* GetActionKeyFilter() const
@@ -90,7 +92,7 @@ public:
 	* @param  MakeActionCallback
 	* @return
 	*/
-	int32 RegisterStructActions(const FMakeStructSpawnerDelegate& MakeActionCallback);
+	UE_API int32 RegisterStructActions(const FMakeStructSpawnerDelegate& MakeActionCallback);
 
 	/// 
 	DECLARE_DELEGATE_RetVal_OneParam(UBlueprintNodeSpawner*, FMakeEnumSpawnerDelegate, const UEnum*);
@@ -100,7 +102,7 @@ public:
 	 * @param  MakeActionCallback
 	 * @return
 	 */
-	int32 RegisterEnumActions(const FMakeEnumSpawnerDelegate& MakeActionCallback);
+	UE_API int32 RegisterEnumActions(const FMakeEnumSpawnerDelegate& MakeActionCallback);
 
 	/// 
 	DECLARE_DELEGATE_RetVal_OneParam(UBlueprintNodeSpawner*, FMakeFuncSpawnerDelegate, const UFunction*);
@@ -123,7 +125,7 @@ private:
 
 	/** Only FBlueprintActionDatabase can spawn and distribute this. */
 	friend class FBlueprintActionDatabase;
-	FBlueprintActionDatabaseRegistrar(FActionRegistry& Database, FUnloadedActionRegistry& UnloadedDatabase, FPrimingQueue& PrimingQueue, TSubclassOf<UEdGraphNode> DefaultKey = nullptr);
+	UE_API FBlueprintActionDatabaseRegistrar(FActionRegistry& Database, FUnloadedActionRegistry& UnloadedDatabase, FPrimingQueue& PrimingQueue, TSubclassOf<UEdGraphNode> DefaultKey = nullptr);
 
 	/**
 	 * 
@@ -132,7 +134,7 @@ private:
 	 * @param  MakeActionCallback    
 	 * @return 
 	 */
-	int32 RegisterClassFactoryActions(const UClass* TargetType, const FMakeFuncSpawnerDelegate& MakeActionCallback);
+	UE_API int32 RegisterClassFactoryActions(const UClass* TargetType, const FMakeFuncSpawnerDelegate& MakeActionCallback);
 
 	/**
 	 * Internal method that actually adds the action to the database.
@@ -140,7 +142,7 @@ private:
 	 * @param  ActionKey	The key you want the action listed under.
 	 * @param  NodeSpawner	The new node spawner that you wish to register.
 	 */
-	bool AddActionToDatabase(UObject const* ActionKey, UBlueprintNodeSpawner* NodeSpawner);
+	UE_API bool AddActionToDatabase(UObject const* ActionKey, UBlueprintNodeSpawner* NodeSpawner);
 
 private:
 	/** The node type that this is currently passed to (acts as a fallback database key). */
@@ -166,3 +168,5 @@ private:
 	/** Cached ptr to asset tools to avoid LoadModule calls per-AddBlueprintAction call */
 	IAssetTools* AssetTools;
 };
+
+#undef UE_API
diff --git a/Engine/Source/Editor/BlueprintGraph/Public/BlueprintActionFilter.h b/Engine/Source/Editor/BlueprintGraph/Public/BlueprintActionFilter.h
--- a/Engine/Source/Editor/BlueprintGraph/Public/BlueprintActionFilter.h
+++ b/Engine/Source/Editor/BlueprintGraph/Public/BlueprintActionFilter.h
@@ -26,6 +26,8 @@
 #include "UObject/WeakObjectPtr.h"
 #include "UObject/WeakObjectPtrTemplates.h"
 
+#define UE_API BLUEPRINTGRAPH_API
+
 #define ENABLE_BLUEPRINT_ACTION_FILTER_PROFILING (!(UE_BUILD_SHIPPING || UE_BUILD_TEST) && 0)
 
 class FProperty;
@@ -91,13 +93,13 @@ struct FBlueprintActionContext
  * UBlueprintNodeSpawner, and caches associated fields/files/etc. as they're 
  * requested (to optimize duplicated queries cross rejection tests).
  */
-struct BLUEPRINTGRAPH_API FBlueprintActionInfo
+struct FBlueprintActionInfo
 {
 	/** */
-	FBlueprintActionInfo(UObject const* ActionOwner, UBlueprintNodeSpawner const* Action);
+	UE_API FBlueprintActionInfo(UObject const* ActionOwner, UBlueprintNodeSpawner const* Action);
 
 	/** */
-	FBlueprintActionInfo(FBlueprintActionInfo const& Rhs, IBlueprintNodeBinder::FBindingSet const& Bindings);
+	UE_API FBlueprintActionInfo(FBlueprintActionInfo const& Rhs, IBlueprintNodeBinder::FBindingSet const& Bindings);
 
 	/**
 	 * Retrieves the key that the wrapped action is associated with in the 
@@ -105,14 +107,14 @@ struct BLUEPRINTGRAPH_API FBlueprintActionInfo
 	 * 
 	 * @return The class/asset that this action conceptually belongs to.
 	 */
-	UObject const* GetActionOwner() const;
+	UE_API UObject const* GetActionOwner() const;
 
 	/**
 	 * Retrieves any bindings that the action will apply to spawned nodes.
 	 * 
 	 * @return The set of bindings tied to this action.
 	 */
-	IBlueprintNodeBinder::FBindingSet const& GetBindings() const;
+	UE_API IBlueprintNodeBinder::FBindingSet const& GetBindings() const;
 
 	/**
 	 * Retrieves a class associated with the wrapped action. Intended to be the 
@@ -122,7 +124,7 @@ struct BLUEPRINTGRAPH_API FBlueprintActionInfo
 	 * 
 	 * @return The class that this action conceptually belongs to (null if the action is keyed to an asset that we cannot derive a class from).
 	 */
-	UClass const* GetOwnerClass();
+	UE_API UClass const* GetOwnerClass();
 
 	/**
 	 * Retrieves the node class that the wrapped action will spawn (assume to 
@@ -130,7 +132,7 @@ struct BLUEPRINTGRAPH_API FBlueprintActionInfo
 	 * 
 	 * @return The node type that the action will spawn.
 	 */
-	UClass const* GetNodeClass();
+	UE_API UClass const* GetNodeClass();
 
 	/**
 	 * Certain actions are associated with specific member fields (a member 
@@ -139,7 +141,7 @@ struct BLUEPRINTGRAPH_API FBlueprintActionInfo
 	 * 
 	 * @return The member field associated with the wrapped action (null if there isn't one).
 	 */
-	FFieldVariant GetAssociatedMemberField();
+	UE_API FFieldVariant GetAssociatedMemberField();
 
 	/**
 	 * Certain actions are associated with specific properties (like delegate  
@@ -149,7 +151,7 @@ struct BLUEPRINTGRAPH_API FBlueprintActionInfo
 	 * 
 	 * @return The property associated with the wrapped action (null if there isn't one).
 	 */
-	FProperty const* GetAssociatedProperty();
+	UE_API FProperty const* GetAssociatedProperty();
 
 	/**
 	 * Certain actions are associated with specific functions (like function    
@@ -158,7 +160,7 @@ struct BLUEPRINTGRAPH_API FBlueprintActionInfo
 	 * 
 	 * @return The function associated with the wrapped action (null if there isn't one).
 	 */
-	UFunction const* GetAssociatedFunction();
+	UE_API UFunction const* GetAssociatedFunction();
 
 	/** The raw action that this struct represent (const so we don't mutate the database) */
 	UBlueprintNodeSpawner const* const NodeSpawner;
@@ -207,14 +209,14 @@ ENUM_CLASS_FLAGS(EActionFilterTestFlags);
 ******************************************************************************/
 
 // @todo - Need to revisit this later; for now this wrapper type remains in place for backcompat, but flags are currently not in use.
-struct BLUEPRINTGRAPH_API FActionFilterTest : TSharedFromThis<FActionFilterTest>
+struct FActionFilterTest : TSharedFromThis<FActionFilterTest>
 {
-	static TAutoConsoleVariable<bool> CVarEnableCaching;
+	static UE_API TAutoConsoleVariable<bool> CVarEnableCaching;
 
 	/** The filter uses a series of rejection tests matching */
 	typedef FBlueprintGraphModule::FActionMenuRejectionTest FRejectionDelegate;
 
-	FActionFilterTest(const FRejectionDelegate Delegate, FString Name, EActionFilterTestFlags Flags = EActionFilterTestFlags::Default);
+	UE_API FActionFilterTest(const FRejectionDelegate Delegate, FString Name, EActionFilterTestFlags Flags = EActionFilterTestFlags::Default);
 
 	TSharedRef<FActionFilterTest> WithFlags(EActionFilterTestFlags InFlags)
 	{
@@ -222,8 +224,8 @@ struct BLUEPRINTGRAPH_API FActionFilterTest : TSharedFromThis<FActionFilterTest>
 		return AsShared();
 	}
 
-	bool Call(FBlueprintActionFilter const& Filter, FBlueprintActionInfo& ActionInfo);
-	bool operator==(const FActionFilterTest& other) const;
+	UE_API bool Call(FBlueprintActionFilter const& Filter, FBlueprintActionInfo& ActionInfo);
+	UE_API bool operator==(const FActionFilterTest& other) const;
 
 	FRejectionDelegate RejectionDelegate;
 	FString Name; // this name must be unique
@@ -239,7 +241,7 @@ struct BLUEPRINTGRAPH_API FActionFilterTest : TSharedFromThis<FActionFilterTest>
  * FBlueprintActionFilter
  ******************************************************************************/
 
-class BLUEPRINTGRAPH_API FBlueprintActionFilter
+class FBlueprintActionFilter
 {
 public:
 	/** The filter uses a series of rejection tests matching */
@@ -289,10 +291,10 @@ public:
 	};
 
 	/** @param InFlags - Filter configuration flags. */
-	FBlueprintActionFilter(const EFlags InFlags = BPFILTER_NoFlags);
+	UE_API FBlueprintActionFilter(const EFlags InFlags = BPFILTER_NoFlags);
 
 	UE_DEPRECATED(5.1, "Please use the version that takes the EFlags type as input")
-	FBlueprintActionFilter(uint32 const Flags);
+	UE_API FBlueprintActionFilter(uint32 const Flags);
 	
 	/**
 	 * Contains the full blueprint/graph/pin context that this is filtering 
@@ -333,10 +335,10 @@ public:
 	TArray<FTargetClassFilterData> TargetClasses;
 
 	/** Helper to add a class to the TargetClasses, TargetClass may already be in the array */
-	static void AddUnique(TArray<FTargetClassFilterData>& ToArray, UClass* TargetClass);
+	static UE_API void AddUnique(TArray<FTargetClassFilterData>& ToArray, UClass* TargetClass);
 
 	/** Helper to add a class to the TargetClasses, fills out FTargetClassFilterData */
-	static void Add(TArray<FTargetClassFilterData>& ToArray, UClass* TargetClass);
+	static UE_API void Add(TArray<FTargetClassFilterData>& ToArray, UClass* TargetClass);
 
 	/** Cached reference to the BluprintGraphModule, which has extra rejection tests: */
 	class FBlueprintGraphModule* BluprintGraphModule;
@@ -351,13 +353,13 @@ public:
 	 *
 	 * @param  RejectionTestDelegate	The rejection test you wish to add to this filter.
 	 */
-	void AddRejectionTest(FRejectionTestDelegate RejectionTestDelegate);
+	UE_API void AddRejectionTest(FRejectionTestDelegate RejectionTestDelegate);
 
 	/**
 	 * @param  RejectionTest	a test this filter will run to cull blueprint actions
 	 * @param  Flags	List of flags that will be added in the provided rejection test
 	 */
-	void AddRejectionTest(TSharedRef<FActionFilterTest> RejectionTest);
+	UE_API void AddRejectionTest(TSharedRef<FActionFilterTest> RejectionTest);
 
 	/**
 	 * Query to check and see if the specified action gets filtered out by this 
@@ -367,7 +369,7 @@ public:
 	 * @param  BlueprintAction	The node-spawner you wish to test.
 	 * @return False if the action passes the filter, otherwise false (the action got filtered out).
 	 */
-	bool IsFiltered(FBlueprintActionInfo& BlueprintAction);
+	UE_API bool IsFiltered(FBlueprintActionInfo& BlueprintAction);
 
 	/**
 	 * Appends another filter to be utilized in IsFiltered() queries, extending  
@@ -379,7 +381,7 @@ public:
 	 * @param  Rhs	The filter to append to this one.
 	 * @return This.
 	 */
-	FBlueprintActionFilter const& operator|=(FBlueprintActionFilter const& Rhs);
+	UE_API FBlueprintActionFilter const& operator|=(FBlueprintActionFilter const& Rhs);
 
 	/**
 	 * Appends another filter to be utilized in IsFiltered() queries, extending  
@@ -391,7 +393,7 @@ public:
 	 * @param  Rhs	The filter to append to this one.
 	 * @return This.
 	 */
-	FBlueprintActionFilter const& operator&=(FBlueprintActionFilter const& Rhs);
+	UE_API FBlueprintActionFilter const& operator&=(FBlueprintActionFilter const& Rhs);
 
 	/**
 	 * @return TRUE if any of the given configuration flags are set on this filter.
@@ -410,9 +412,9 @@ public:
 	}
 
 #if ENABLE_BLUEPRINT_ACTION_FILTER_PROFILING
-	static bool IsFilterTestTraceLoggingEnabled();
-	static bool IsFilterTestStatsLoggingEnabled();
-	TArray<FString> GetFilterTestProfile();
+	static UE_API bool IsFilterTestTraceLoggingEnabled();
+	static UE_API bool IsFilterTestStatsLoggingEnabled();
+	UE_API TArray<FString> GetFilterTestProfile();
 #endif	// ENABLE_BLUEPRINT_ACTION_FILTER_PROFILING
 
 private:
@@ -423,7 +425,7 @@ private:
 	 * @param  BlueprintAction	The node-spawner you wish to test.
 	 * @return False if the action passes the filter, otherwise true (the action got filtered out).
 	 */
-	bool IsFilteredByThis(FBlueprintActionInfo& BlueprintAction) const;
+	UE_API bool IsFilteredByThis(FBlueprintActionInfo& BlueprintAction) const;
 
 	/** Configuration flags for this filter. */
 	EFlags FilterFlags = BPFILTER_NoFlags;
@@ -449,9 +451,11 @@ private:
 
 	mutable TMap<int32, FFilterTestProfileRecord> FilterTestProfiles;
 
-	FFilterTestProfileRecord* BeginFilterTestProfileEvent(int32 FilterTestIndex, uint64 FilterTestFuncPtr) const;
-	void EndFilterTestProfileEvent(FFilterTestProfileRecord* InEvent, bool bWasFiltered, double StartTime) const;
+	UE_API FFilterTestProfileRecord* BeginFilterTestProfileEvent(int32 FilterTestIndex, uint64 FilterTestFuncPtr) const;
+	UE_API void EndFilterTestProfileEvent(FFilterTestProfileRecord* InEvent, bool bWasFiltered, double StartTime) const;
 #endif	// ENABLE_BLUEPRINT_ACTION_FILTER_PROFILING
 };
 
-ENUM_CLASS_FLAGS(FBlueprintActionFilter::EFlags);
\ No newline at end of file
+ENUM_CLASS_FLAGS(FBlueprintActionFilter::EFlags);
+
+#undef UE_API
diff --git a/Engine/Source/Editor/BlueprintGraph/Public/BlueprintEditorSettings.h b/Engine/Source/Editor/BlueprintGraph/Public/BlueprintEditorSettings.h
--- a/Engine/Source/Editor/BlueprintGraph/Public/BlueprintEditorSettings.h
+++ b/Engine/Source/Editor/BlueprintGraph/Public/BlueprintEditorSettings.h
@@ -29,6 +29,8 @@
 
 #include "BlueprintEditorSettings.generated.h"
 
+#define UE_API BLUEPRINTGRAPH_API
+
 class UObject;
 struct FAssetData;
 struct FBPEditorBookmarkNode;
@@ -48,7 +50,7 @@ enum ESaveOnCompile : int
 *	See FKismetDebugUtilities for helper functions
 */
 USTRUCT()
-struct BLUEPRINTGRAPH_API FPerBlueprintSettings 
+struct FPerBlueprintSettings 
 {
 	GENERATED_BODY()
 	
@@ -73,7 +75,7 @@ template<> struct TStructOpsTypeTraits<FPerBlueprintSettings> : public TStructOp
 };
 
 USTRUCT()
-struct BLUEPRINTGRAPH_API FAdditionalBlueprintCategory
+struct FAdditionalBlueprintCategory
 {
 	GENERATED_BODY()
 
@@ -84,8 +86,8 @@ struct BLUEPRINTGRAPH_API FAdditionalBlueprintCategory
 	FSoftClassPath ClassFilter;
 };
 
-UCLASS(config=EditorPerProjectUserSettings)
-class BLUEPRINTGRAPH_API UBlueprintEditorSettings : public UDeveloperSettings
+UCLASS(MinimalAPI, config=EditorPerProjectUserSettings)
+class UBlueprintEditorSettings : public UDeveloperSettings
 {
 	GENERATED_UCLASS_BODY()
 
@@ -373,35 +375,35 @@ public:
 	using FOnIsFunctionAllowedForAsset = TDelegate<bool(const UBlueprint* InBlueprint, const FName InFunctionName)>;
 	
 	/** Delegates called to determine whether a class type is allowed */
-	void RegisterIsClassAllowedDelegate(const FName OwnerName, FOnIsClassAllowed Delegate);
-	void UnregisterIsClassAllowedDelegate(const FName OwnerName);
-	bool IsClassAllowed(const UClass* InClass) const;
+	UE_API void RegisterIsClassAllowedDelegate(const FName OwnerName, FOnIsClassAllowed Delegate);
+	UE_API void UnregisterIsClassAllowedDelegate(const FName OwnerName);
+	UE_API bool IsClassAllowed(const UClass* InClass) const;
 	bool HasClassFiltering() const { return IsClassAllowedDelegates.Num() > 0; }
 
-	void RegisterIsClassPathAllowedDelegate(const FName OwnerName, FOnIsClassPathAllowed Delegate);
-	void UnregisterIsClassPathAllowedDelegate(const FName OwnerName);
-	bool IsClassPathAllowed(const FTopLevelAssetPath& InClassPath) const;
+	UE_API void RegisterIsClassPathAllowedDelegate(const FName OwnerName, FOnIsClassPathAllowed Delegate);
+	UE_API void UnregisterIsClassPathAllowedDelegate(const FName OwnerName);
+	UE_API bool IsClassPathAllowed(const FTopLevelAssetPath& InClassPath) const;
 	bool HasClassPathFiltering() const { return IsClassPathAllowedDelegates.Num() > 0; }
 
-	void RegisterIsClassAllowedOnPinDelegate(const FName OwnerName, FOnIsClassAllowed Delegate);
-	void UnregisterIsClassAllowedOnPinDelegate(const FName OwnerName);
-	bool IsClassAllowedOnPin(const UClass* InClass) const;
+	UE_API void RegisterIsClassAllowedOnPinDelegate(const FName OwnerName, FOnIsClassAllowed Delegate);
+	UE_API void UnregisterIsClassAllowedOnPinDelegate(const FName OwnerName);
+	UE_API bool IsClassAllowedOnPin(const UClass* InClass) const;
 	bool HasClassOnPinFiltering() const { return IsClassAllowedOnPinDelegates.Num() > 0; }
 
-	void RegisterIsClassPathAllowedOnPinDelegate(const FName OwnerName, FOnIsClassPathAllowed Delegate);
-	void UnregisterIsClassPathAllowedOnPinDelegate(const FName OwnerName);
-	bool IsClassPathAllowedOnPin(const FTopLevelAssetPath& InClassPath) const;
+	UE_API void RegisterIsClassPathAllowedOnPinDelegate(const FName OwnerName, FOnIsClassPathAllowed Delegate);
+	UE_API void UnregisterIsClassPathAllowedOnPinDelegate(const FName OwnerName);
+	UE_API bool IsClassPathAllowedOnPin(const FTopLevelAssetPath& InClassPath) const;
 	bool HasClassPathOnPinFiltering() const { return IsClassPathAllowedOnPinDelegates.Num() > 0; }
 
 	/** Delegates called to determine whether a function is allowed for a given Blueprint asset */
-	void RegisterIsFunctionAllowedForAssetDelegate(const FName OwnerName, FOnIsFunctionAllowedForAsset Delegate);
-	void UnregisterIsFunctionAllowedForAssetDelegate(const FName OwnerName);
+	UE_API void RegisterIsFunctionAllowedForAssetDelegate(const FName OwnerName, FOnIsFunctionAllowedForAsset Delegate);
+	UE_API void UnregisterIsFunctionAllowedForAssetDelegate(const FName OwnerName);
 
-	bool IsFunctionAllowed(const UBlueprint* InBlueprint, const FName InFunctionName) const;
+	UE_API bool IsFunctionAllowed(const UBlueprint* InBlueprint, const FName InFunctionName) const;
 
-	bool IsFunctionAllowedForAsset(const UBlueprint* InBlueprint, const FName InFunctionName) const;
+	UE_API bool IsFunctionAllowedForAsset(const UBlueprint* InBlueprint, const FName InFunctionName) const;
 
-	static bool IsFunctionAllowedForPermissions(const UBlueprint* InBlueprint, const FName FunctionName, const FPathPermissionList& InFunctionPermissions);
+	static UE_API bool IsFunctionAllowedForPermissions(const UBlueprint* InBlueprint, const FName FunctionName, const FPathPermissionList& InFunctionPermissions);
 	
 private:
 	/** All function permissions */
@@ -429,11 +431,13 @@ private:
 
 protected:
 	//~ Begin UObject Interface
-	virtual void PostInitProperties();
-	virtual void PostEditChangeProperty(struct FPropertyChangedEvent& PropertyChangedEvent) override;
+	UE_API virtual void PostInitProperties();
+	UE_API virtual void PostEditChangeProperty(struct FPropertyChangedEvent& PropertyChangedEvent) override;
 	//~ End UObject Interface
 	
-	void OnAssetRenamed(FAssetData const& AssetInfo, const FString& InOldName);
+	UE_API void OnAssetRenamed(FAssetData const& AssetInfo, const FString& InOldName);
 	
-	void OnAssetRemoved(UObject* Object);
+	UE_API void OnAssetRemoved(UObject* Object);
 };
+
+#undef UE_API
diff --git a/Engine/Source/Editor/BlueprintGraph/Public/BlueprintGraphModule.h b/Engine/Source/Editor/BlueprintGraph/Public/BlueprintGraphModule.h
--- a/Engine/Source/Editor/BlueprintGraph/Public/BlueprintGraphModule.h
+++ b/Engine/Source/Editor/BlueprintGraph/Public/BlueprintGraphModule.h
@@ -26,7 +26,7 @@ public:
 	 * 
 	 * @return An array of additional rejection tests for FBlueprintActionFilter.
 	 */
-	BLUEPRINTGRAPH_API TArray<FActionMenuRejectionTest>& GetExtendedActionMenuFilters() { return ExtendedMenuFilters; }
+	TArray<FActionMenuRejectionTest>& GetExtendedActionMenuFilters() { return ExtendedMenuFilters; }
 
 	virtual void ShutdownModule() override;
 	BLUEPRINTGRAPH_API void RegisterGraphAction(const UClass* AssetType, TUniquePtr<FAssetBlueprintGraphActions> Action);
diff --git a/Engine/Source/Editor/BlueprintGraph/Public/BlueprintGraphPanelPinFactory.h b/Engine/Source/Editor/BlueprintGraph/Public/BlueprintGraphPanelPinFactory.h
--- a/Engine/Source/Editor/BlueprintGraph/Public/BlueprintGraphPanelPinFactory.h
+++ b/Engine/Source/Editor/BlueprintGraph/Public/BlueprintGraphPanelPinFactory.h
@@ -6,9 +6,13 @@
 #include "EdGraphUtilities.h"
 #include "Templates/SharedPointer.h"
 
+#define UE_API BLUEPRINTGRAPH_API
+
 class SGraphPin;
 
-class BLUEPRINTGRAPH_API FBlueprintGraphPanelPinFactory: public FGraphPanelPinFactory
+class FBlueprintGraphPanelPinFactory: public FGraphPanelPinFactory
 {
-	virtual TSharedPtr<class SGraphPin> CreatePin(class UEdGraphPin* InPin) const override;
+	UE_API virtual TSharedPtr<class SGraphPin> CreatePin(class UEdGraphPin* InPin) const override;
 };
+
+#undef UE_API
diff --git a/Engine/Source/Editor/BlueprintGraph/Public/BlueprintNodeSignature.h b/Engine/Source/Editor/BlueprintGraph/Public/BlueprintNodeSignature.h
--- a/Engine/Source/Editor/BlueprintGraph/Public/BlueprintNodeSignature.h
+++ b/Engine/Source/Editor/BlueprintGraph/Public/BlueprintNodeSignature.h
@@ -12,14 +12,16 @@
 #include "UObject/NameTypes.h"
 #include "UObject/Object.h"
 
+#define UE_API BLUEPRINTGRAPH_API
+
 class UEdGraphNode;
 
-struct BLUEPRINTGRAPH_API FBlueprintNodeSignature
+struct FBlueprintNodeSignature
 {
 public:
 	FBlueprintNodeSignature() {}
-	FBlueprintNodeSignature(FString const& UserString);
-	FBlueprintNodeSignature(TSubclassOf<UEdGraphNode> NodeClass);
+	UE_API FBlueprintNodeSignature(FString const& UserString);
+	UE_API FBlueprintNodeSignature(TSubclassOf<UEdGraphNode> NodeClass);
 
 	/**
 	 * 
@@ -27,7 +29,7 @@ public:
 	 * @param  NodeClass	
 	 * @return 
 	 */
-	void SetNodeClass(TSubclassOf<UEdGraphNode> NodeClass);
+	UE_API void SetNodeClass(TSubclassOf<UEdGraphNode> NodeClass);
 
 	/**
 	 * 
@@ -35,7 +37,7 @@ public:
 	 * @param  SignatureObj	
 	 * @return 
 	 */
-	void AddSubObject(FFieldVariant SignatureObj);
+	UE_API void AddSubObject(FFieldVariant SignatureObj);
 
 	/**
 	 * 
@@ -43,7 +45,7 @@ public:
 	 * @param  Value	
 	 * @return 
 	 */
-	void AddKeyValue(FString const& KeyValue);
+	UE_API void AddKeyValue(FString const& KeyValue);
 
 	/**
 	 * 
@@ -52,28 +54,28 @@ public:
 	 * @param  Value	
 	 * @return 
 	 */
-	void AddNamedValue(FName SignatureKey, FString const& Value);
+	UE_API void AddNamedValue(FName SignatureKey, FString const& Value);
 
 	/**
 	 * 
 	 * 
 	 * @return 
 	 */
-	bool IsValid() const;
+	UE_API bool IsValid() const;
 
 	/**
 	 * 
 	 * 
 	 * @return 
 	 */
-	FString const& ToString() const;
+	UE_API FString const& ToString() const;
 
 	/**
 	*
 	*
 	* @return
 	*/
-	FGuid const& AsGuid() const;
+	UE_API FGuid const& AsGuid() const;
 
 private:
 	/**
@@ -81,7 +83,7 @@ private:
 	 * 
 	 * @return 
 	 */
-	void MarkDirty();
+	UE_API void MarkDirty();
 
 private:
 	typedef TMap<FName, FString> FSignatureSet;
@@ -91,3 +93,4 @@ private:
 	mutable FString CachedSignatureString;
 };
 
+#undef UE_API
diff --git a/Engine/Source/Editor/BlueprintGraph/Public/BlueprintNodeSpawner.h b/Engine/Source/Editor/BlueprintGraph/Public/BlueprintNodeSpawner.h
--- a/Engine/Source/Editor/BlueprintGraph/Public/BlueprintNodeSpawner.h
+++ b/Engine/Source/Editor/BlueprintGraph/Public/BlueprintNodeSpawner.h
@@ -22,6 +22,8 @@
 
 #include "BlueprintNodeSpawner.generated.h"
 
+#define UE_API BLUEPRINTGRAPH_API
+
 class UEdGraph;
 class UEdGraphNode;
 
@@ -110,8 +112,8 @@ private:
  * different node types). That way, we get the type inference that comes with  
  * UObjects (and we don't have to continuously compare identification strings). 
  */
-UCLASS(Transient)
-class BLUEPRINTGRAPH_API UBlueprintNodeSpawner : public UObject, public IBlueprintNodeBinder
+UCLASS(MinimalAPI, Transient)
+class UBlueprintNodeSpawner : public UObject, public IBlueprintNodeBinder
 {
 	GENERATED_UCLASS_BODY()
 	DECLARE_DELEGATE_TwoParams(FCustomizeNodeDelegate, UEdGraphNode*, bool);
@@ -128,7 +130,7 @@ public:
 	 * @param  PostSpawnDelegate    A delegate to perform specialized node setup post-spawn.
 	 * @return A newly allocated instance of this class.
 	 */
-	static UBlueprintNodeSpawner* Create(TSubclassOf<UEdGraphNode> const NodeClass, UObject* Outer = nullptr, FCustomizeNodeDelegate PostSpawnDelegate = FCustomizeNodeDelegate());
+	static UE_API UBlueprintNodeSpawner* Create(TSubclassOf<UEdGraphNode> const NodeClass, UObject* Outer = nullptr, FCustomizeNodeDelegate PostSpawnDelegate = FCustomizeNodeDelegate());
 	
 	/**
 	 * Templatized version of the above Create() method (where we specify the 
@@ -156,7 +158,7 @@ public:
 	FUiSpecOverrideDelegate DynamicUiSignatureGetter;
 
 	// UObject interface
-	virtual void BeginDestroy() override;
+	UE_API virtual void BeginDestroy() override;
 	// End UObject interface
 
 	/**
@@ -164,7 +166,7 @@ public:
 	 * operations. Will cache a node-template (via GetTemplateNode), along with  
 	 * any expensive text strings, to avoid constructing them all on demand.
 	 */
-	virtual void Prime();
+	UE_API virtual void Prime();
 
 	/**
 	 * Takes the FBlueprintActionUiSpec that this was spawned with and attempts 
@@ -173,7 +175,7 @@ public:
 	 * @param  TargetGraph	Optional context for helping spawn the template node. 
 	 * @return This spawner's default ui spec.
 	 */
-	FBlueprintActionUiSpec const& PrimeDefaultUiSpec(UEdGraph* TargetGraph = nullptr) const;
+	UE_API FBlueprintActionUiSpec const& PrimeDefaultUiSpec(UEdGraph* TargetGraph = nullptr) const;
 
 	/**
 	 * We want to be able to compare spawners, and have a signature that is 
@@ -183,7 +185,7 @@ public:
 	 * 
 	 * @return A set of object-paths/names that distinguish this spawner from others.
 	 */
-	virtual FBlueprintNodeSignature GetSpawnerSignature() const;
+	UE_API virtual FBlueprintNodeSignature GetSpawnerSignature() const;
 
 	/**
 	 * Takes the default FBlueprintActionUiSpec and modifies it dynamically to 
@@ -193,7 +195,7 @@ public:
 	 * @param  Bindings		Bindings that will be applied to the node post-spawn.
 	 * @return A ui spec, detailing how to present this action in the menu.
 	 */
-	virtual FBlueprintActionUiSpec GetUiSpec(FBlueprintActionContext const& Context, FBindingSet const& Bindings) const;
+	UE_API virtual FBlueprintActionUiSpec GetUiSpec(FBlueprintActionContext const& Context, FBindingSet const& Bindings) const;
 
 	/**
 	 * Takes care of spawning a node for the specified graph. Looks to see if 
@@ -210,7 +212,7 @@ public:
 	 * @param  Location     Where you want the new node positioned in the graph.
 	 * @return Null if it failed to spawn a node, otherwise a newly spawned node or possibly one that already existed.
 	 */
-	virtual UEdGraphNode* Invoke(UEdGraph* ParentGraph, FBindingSet const& Bindings, FVector2D const Location) const;
+	UE_API virtual UEdGraphNode* Invoke(UEdGraph* ParentGraph, FBindingSet const& Bindings, FVector2D const Location) const;
 
 	/**
 	 * Retrieves a cached template for the node that this is set to spawn. Will
@@ -218,7 +220,7 @@ public:
 	 *
 	 * @return The cached template-node (if one already exists for this spawner).
 	 */
-	UEdGraphNode* GetCachedTemplateNode() const;
+	UE_API UEdGraphNode* GetCachedTemplateNode() const;
 
 	/**
 	 * Retrieves a cached template for the node that this is set to spawn. Will
@@ -230,19 +232,19 @@ public:
 	 * @param  Bindings    	  Objects to bind to the template node
 	 * @return Should return a new/cached template-node (but could be null, or some pre-existing node... depends on the sub-class's Invoke() method).
 	 */
-	UEdGraphNode* GetTemplateNode(UEdGraph* TargetGraph = nullptr, FBindingSet const& Bindings = FBindingSet()) const;
+	UE_API UEdGraphNode* GetTemplateNode(UEdGraph* TargetGraph = nullptr, FBindingSet const& Bindings = FBindingSet()) const;
 
 	/**
 	 * Removes the spawner's cached template node (if it has one), meaning that
 	 * the next GetTemplateNode() call will spawn a new one (and that calls to 
 	 * GetCachedTemplateNode() will return null).
 	 */
-	void ClearCachedTemplateNode() const;
+	UE_API void ClearCachedTemplateNode() const;
 
 	/**
 	 * Determine if the node of this spawner's template's type should be filtered in the actions menu
 	 */
-	virtual bool IsTemplateNodeFilteredOut(FBlueprintActionFilter const& Filter) const;
+	UE_API virtual bool IsTemplateNodeFilteredOut(FBlueprintActionFilter const& Filter) const;
 
 	// IBlueprintNodeBinder interface
 	virtual bool IsBindingCompatible(FBindingObject BindingCandidate) const override { return false; }
@@ -292,7 +294,7 @@ private:
 	 * @param  PostSpawnDelegate	A delegate to run after spawning the node, but prior to allocating the node's pins.
 	 * @return Null if it failed to spawn a node (if NodeClass is null), otherwise a newly spawned node.
 	 */
-	UEdGraphNode* SpawnEdGraphNode(TSubclassOf<UEdGraphNode> NodeClass, UEdGraph* ParentGraph, const FBindingSet& Bindings, FVector2D Location, FCustomizeNodeDelegate PostSpawnDelegate) const;
+	UE_API UEdGraphNode* SpawnEdGraphNode(TSubclassOf<UEdGraphNode> NodeClass, UEdGraph* ParentGraph, const FBindingSet& Bindings, FVector2D Location, FCustomizeNodeDelegate PostSpawnDelegate) const;
 };
 
 /*******************************************************************************
@@ -319,3 +321,5 @@ NodeType* UBlueprintNodeSpawner::SpawnNode(UEdGraph* ParentGraph, FBindingSet co
 {
 	return SpawnNode<NodeType>(NodeType::StaticClass(), ParentGraph, Bindings, Location, PostSpawnDelegate);
 }
+
+#undef UE_API
diff --git a/Engine/Source/Editor/BlueprintGraph/Public/BlueprintNodeStatics.h b/Engine/Source/Editor/BlueprintGraph/Public/BlueprintNodeStatics.h
--- a/Engine/Source/Editor/BlueprintGraph/Public/BlueprintNodeStatics.h
+++ b/Engine/Source/Editor/BlueprintGraph/Public/BlueprintNodeStatics.h
@@ -5,16 +5,20 @@
 #include "CoreMinimal.h"
 #include "KismetCompilerMisc.h"
 
+#define UE_API BLUEPRINTGRAPH_API
+
 class UEdGraphPin;
 class UFunction;
 class UK2Node;
 template <typename FuncType> class TFunctionRef;
 
-class BLUEPRINTGRAPH_API FBlueprintNodeStatics
+class FBlueprintNodeStatics
 {
 public:
 
-	static UEdGraphPin* CreateSelfPin(UK2Node* Node, const UFunction* Function);
+	static UE_API UEdGraphPin* CreateSelfPin(UK2Node* Node, const UFunction* Function);
+
+	static UE_API bool CreateParameterPinsForFunction(UK2Node* Node, const UFunction* Function, TFunctionRef<void(UEdGraphPin* /*Pin*/)> PostParameterPinCreatedCallback);
+};
 
-	static bool CreateParameterPinsForFunction(UK2Node* Node, const UFunction* Function, TFunctionRef<void(UEdGraphPin* /*Pin*/)> PostParameterPinCreatedCallback);
-};
\ No newline at end of file
+#undef UE_API
diff --git a/Engine/Source/Editor/BlueprintGraph/Public/BlueprintNodeTemplateCache.h b/Engine/Source/Editor/BlueprintGraph/Public/BlueprintNodeTemplateCache.h
--- a/Engine/Source/Editor/BlueprintGraph/Public/BlueprintNodeTemplateCache.h
+++ b/Engine/Source/Editor/BlueprintGraph/Public/BlueprintNodeTemplateCache.h
@@ -10,6 +10,8 @@
 #include "CoreTypes.h"
 #include "UObject/GCObject.h"
 
+#define UE_API BLUEPRINTGRAPH_API
+
 class FReferenceCollector;
 class UBlueprint;
 class UBlueprintNodeSpawner;
@@ -23,11 +25,11 @@ class UEdGraphNode;
  * instantiating a bunch of graphs/blueprints per UBlueprintNodeSpawner, we'd 
  * rather have a small centralized set here.
  */
-class BLUEPRINTGRAPH_API FBlueprintNodeTemplateCache : public FGCObject
+class FBlueprintNodeTemplateCache : public FGCObject
 {
 public:
 	/** */
-	FBlueprintNodeTemplateCache();
+	UE_API FBlueprintNodeTemplateCache();
 
 	/**
 	 * Retrieves a cached template associated with the supplied spawner. Will
@@ -39,7 +41,7 @@ public:
 	 * @param  TargetGraph	Optional param that defines a compatible graph outer (used as an achetype if we don't have a compatible outer on hand).
 	 * @return Should return a new/cached template-node (but could be null, or some pre-existing node... depends on the sub-class's Invoke() method).
 	 */
-	UEdGraphNode* GetNodeTemplate(UBlueprintNodeSpawner const* NodeSpawner, UEdGraph* TargetGraph = nullptr);
+	UE_API UEdGraphNode* GetNodeTemplate(UBlueprintNodeSpawner const* NodeSpawner, UEdGraph* TargetGraph = nullptr);
 
 	/**
 	 * Retrieves a cached template associated with the supplied spawner. Does 
@@ -49,7 +51,7 @@ public:
 	 * @param  ENoInit Signifies to use this function over the other (mutating) GetNodeTemplate().
 	 * @return Should return the cached template-node (if one exists, otherwise false).
 	 */
-	UEdGraphNode* GetNodeTemplate(UBlueprintNodeSpawner const* NodeSpawner, ENoInit) const;
+	UE_API UEdGraphNode* GetNodeTemplate(UBlueprintNodeSpawner const* NodeSpawner, ENoInit) const;
 
 	/**
 	 * Wipes any nodes that were cached on behalf of the specified spawner 
@@ -58,7 +60,7 @@ public:
 	 * 
 	 * @param  NodeSpawner	The spawner we want cached node-templates cleared for.
 	 */
-	void ClearCachedTemplate(UBlueprintNodeSpawner const* NodeSpawner);
+	UE_API void ClearCachedTemplate(UBlueprintNodeSpawner const* NodeSpawner);
 
 	/**
 	 * Utility method to help external systems identify if a graph they have 
@@ -67,7 +69,7 @@ public:
 	 * @param  ParentGraph	The graph you want to check.
 	 * @return True if the graph was allocated by a FBlueprintNodeTemplateCache (to house template nodes).
 	 */
-	static bool IsTemplateOuter(UEdGraph* ParentGraph);
+	static UE_API bool IsTemplateOuter(UEdGraph* ParentGraph);
 
 	/**
 	 * Approximates the current memory footprint of the entire cache 
@@ -75,7 +77,7 @@ public:
 	 * 
 	 * @return The approximated total (in bytes) that this cache has allocated.
 	 */
-	int64 GetEstimateCacheSize() const;
+	UE_API int64 GetEstimateCacheSize() const;
 
 	/**
 	 * External systems can make changes that alter the memory footprint of the
@@ -85,11 +87,11 @@ public:
 	 * 
 	 * @return The new approximated total (in bytes) that this cache has allocated.
 	 */
-	int64 RecalculateCacheSize();
+	UE_API int64 RecalculateCacheSize();
 
 	// FGCObject interface
-	virtual void AddReferencedObjects(FReferenceCollector& Collector) override;
-	virtual FString GetReferencerName() const override;
+	UE_API virtual void AddReferencedObjects(FReferenceCollector& Collector) override;
+	UE_API virtual FString GetReferencerName() const override;
 	// End FGCObject interface
 
 private:
@@ -102,7 +104,7 @@ private:
 	 * @param  Blueprint	The blueprint you wish to store for reuse.
 	 * @return True if the blueprint was successfully cached (the cache could be full, and therefore this could fail).
 	 */
-	bool CacheBlueprintOuter(UBlueprint* Blueprint);
+	UE_API bool CacheBlueprintOuter(UBlueprint* Blueprint);
 
 	/**
 	 * Attempts to cache the supplied node, and associates it with the specified 
@@ -112,7 +114,7 @@ private:
 	 * @param  NewNode		The template node you want stored.
 	 * @return True if the node was successfully cached (the cache could be full, and therefore this could fail).
 	 */
-	bool CacheTemplateNode(UBlueprintNodeSpawner const* NodeSpawner, UEdGraphNode* NewNode);
+	UE_API bool CacheTemplateNode(UBlueprintNodeSpawner const* NodeSpawner, UEdGraphNode* NewNode);
 
 private:
 	/** 
@@ -138,3 +140,5 @@ private:
 	 */
 	int64 ApproximateObjectMem;
 };
+
+#undef UE_API
