UE-294487 - UK2Node_VariableGet::SetPurity incorrectly removed

CL 41329753 erroneously removed the SetPurity function from the public API of UK2Node_VariableGet. This fix re-adds SetPurity, but changes its implementation to use TogglePurity.

diff --git a/Engine/Source/Editor/BlueprintGraph/Classes/K2Node_VariableGet.h b/Engine/Source/Editor/BlueprintGraph/Classes/K2Node_VariableGet.h
--- a/Engine/Source/Editor/BlueprintGraph/Classes/K2Node_VariableGet.h
+++ b/Engine/Source/Editor/BlueprintGraph/Classes/K2Node_VariableGet.h
@@ -66,6 +66,14 @@ public:
 	static UE_API FText GetPropertyTooltip(const FProperty* VariableProperty);
 	static UE_API FText GetBlueprintVarTooltip(const FBPVariableDescription& VarDesc);
 
+	/**
+	 * Will change the node's purity, and reallocate pins accordingly (adding/
+	 * removing exec pins).
+	 *
+	 * @param  bNewPurity  The new value for bNewPurity.
+	 */
+	void SetPurity(bool bNewPurity);
+
 private:
 
 	/** Adds pins required for the node to function in a impure manner */
diff --git a/Engine/Source/Editor/BlueprintGraph/Private/K2Node_VariableGet.cpp b/Engine/Source/Editor/BlueprintGraph/Private/K2Node_VariableGet.cpp
--- a/Engine/Source/Editor/BlueprintGraph/Private/K2Node_VariableGet.cpp
+++ b/Engine/Source/Editor/BlueprintGraph/Private/K2Node_VariableGet.cpp
@@ -126,6 +126,19 @@ UK2Node_VariableGet::UK2Node_VariableGet(const FObjectInitializer& ObjectInitial
 {
 }
 
+void UK2Node_VariableGet::SetPurity(bool bNewPurity)
+{
+	const UEdGraphSchema_K2* K2Schema = CastChecked<UEdGraphSchema_K2>(GetSchema());
+	const UEdGraphPin* ValuePin = GetValuePin();
+	const EGetNodeVariation SupportedVariation = K2Node_VariableGetImpl::GetNodeVariation(ValuePin->PinType);
+	const EGetNodeVariation DesiredVariation = bNewPurity ? EGetNodeVariation::Pure : SupportedVariation;
+
+	if (CurrentVariation != DesiredVariation)
+	{
+		TogglePurity(DesiredVariation);
+	}
+}
+
 void UK2Node_VariableGet::CreateImpurePins(TArray<UEdGraphPin*>* InOldPinsPtr)
 {
 	const UEdGraphSchema_K2* K2Schema = CastChecked<UEdGraphSchema_K2>(GetSchema());
