Several fixes to how World Context Objects are handled in blueprints:
Allow calling BlueprintFunctionLibrary functions from classes that do not implement GetWorld or ShowWorldContextPin, it now always shows the pins if needed.
Add ShouldShowWorldContextPin to BlueprintEditorUtils and change the construct and subsystem nodes to use it, and hide the asset picker.
Automatically connect invisible function library world context pins, even if the function says it is optional.

diff --git a/Engine/Source/Editor/BlueprintGraph/Classes/K2Node_GetSubsystem.h b/Engine/Source/Editor/BlueprintGraph/Classes/K2Node_GetSubsystem.h
--- a/Engine/Source/Editor/BlueprintGraph/Classes/K2Node_GetSubsystem.h
+++ b/Engine/Source/Editor/BlueprintGraph/Classes/K2Node_GetSubsystem.h
@@ -40,6 +40,7 @@ public:
 	UE_API void Initialize(UClass* NodeClass);
 	//~Begin UEdGraphNode interface.
 	UE_API virtual void AllocateDefaultPins() override;
+	UE_API virtual FString GetPinMetaData(FName InPinName, FName InKey) override;
 	UE_API virtual bool IsCompatibleWithGraph(const UEdGraph* TargetGraph) const override;
 	UE_API virtual FSlateIcon GetIconAndTint(FLinearColor& OutColor) const override;
 	UE_API virtual FLinearColor GetNodeTitleColor() const override;
diff --git a/Engine/Source/Editor/BlueprintGraph/Private/EdGraphSchema_K2.cpp b/Engine/Source/Editor/BlueprintGraph/Private/EdGraphSchema_K2.cpp
--- a/Engine/Source/Editor/BlueprintGraph/Private/EdGraphSchema_K2.cpp
+++ b/Engine/Source/Editor/BlueprintGraph/Private/EdGraphSchema_K2.cpp
@@ -816,31 +816,7 @@ bool UEdGraphSchema_K2::CanFunctionBeUsedInGraph(const UClass* InClass, const UF
 			return false;
 		}
 
-		const bool bRequiresWorldContext = InFunction->HasMetaData(FBlueprintMetadata::MD_WorldContext);
-		if (bRequiresWorldContext)
-		{
-			if (InDestGraph && !InFunction->HasMetaData(FBlueprintMetadata::MD_CallableWithoutWorldContext))
-			{
-				const FString& ContextParam = InFunction->GetMetaData(FBlueprintMetadata::MD_WorldContext);
-				if (InFunction->FindPropertyByName(FName(*ContextParam)) != nullptr)
-				{
-					UBlueprint* BP = FBlueprintEditorUtils::FindBlueprintForGraph(InDestGraph);
-					const bool bIsFunctLib = BP && (EBlueprintType::BPTYPE_FunctionLibrary == BP->BlueprintType);
-					const bool bIsMacroLib = BP && (EBlueprintType::BPTYPE_MacroLibrary == BP->BlueprintType);
-					UClass* ParentClass = BP ? BP->ParentClass : NULL;
-					const bool bIncompatibleParent = ParentClass && (!FBlueprintEditorUtils::ImplementsGetWorld(BP) && !ParentClass->HasMetaDataHierarchical(FBlueprintMetadata::MD_ShowWorldContextPin));
-					if (!bIsMacroLib && !bIsFunctLib && bIncompatibleParent)
-					{
-						if (OutReason != nullptr)
-						{
-							*OutReason = LOCTEXT("FunctionRequiresWorldContext", "Function requires a world context.");
-						}
-
-						return false;
-					}
-				}	
-			}
-		}
+		// We no longer block functions that require world contexts as it will add the pin when needed
 
 		const bool bFunctionStatic = InFunction->HasAllFunctionFlags(FUNC_Static);
 		const bool bHasReturnParams = (InFunction->GetReturnProperty() != NULL);
diff --git a/Engine/Source/Editor/BlueprintGraph/Private/K2Node_CallFunction.cpp b/Engine/Source/Editor/BlueprintGraph/Private/K2Node_CallFunction.cpp
--- a/Engine/Source/Editor/BlueprintGraph/Private/K2Node_CallFunction.cpp
+++ b/Engine/Source/Editor/BlueprintGraph/Private/K2Node_CallFunction.cpp
@@ -2304,26 +2304,29 @@ void UK2Node_CallFunction::ValidateNodeDuringCompilation(class FCompilerResultsL
 		// enforce WorldContext restrictions
 		const bool bInsideBpFuncLibrary = Blueprint && (BPTYPE_FunctionLibrary == Blueprint->BlueprintType);
 		
-		// go through all of the pins and verify if we have a visible world context or not.
-		bool bVisibleWorldContext = false;
-		for (const UEdGraphPin* Pin : Pins)
-		{
-			if (!Pin->bHidden && Pin->PinName.IsEqual(FName("WorldContextObject")))
-			{
-				bVisibleWorldContext = true;
-				break;
-			}
-		}
 		if (!bInsideBpFuncLibrary &&
 			Function->HasMetaData(FBlueprintMetadata::MD_WorldContext) && 
-			(!Function->HasMetaData(FBlueprintMetadata::MD_CallableWithoutWorldContext) && !bVisibleWorldContext))
+			(!Function->HasMetaData(FBlueprintMetadata::MD_CallableWithoutWorldContext)))
 		{
 			check(Blueprint);
 			UClass* ParentClass = Blueprint->ParentClass;
 			check(ParentClass);
-			if (ParentClass && !FBlueprintEditorUtils::ImplementsGetWorld(Blueprint) && !ParentClass->HasMetaDataHierarchical(FBlueprintMetadata::MD_ShowWorldContextPin))
+
+			// go through all of the pins and verify if we have a connected world context or not.
+			const FName ContextParam = FName(*Function->GetMetaData(FBlueprintMetadata::MD_WorldContext));
+			bool bConnectedWorldContext = false;
+			for (const UEdGraphPin* Pin : Pins)
 			{
-				MessageLog.Warning(*LOCTEXT("FunctionUnsafeInContext", "Function '@@' is unsafe to call from blueprints of class '@@'.").ToString(), this, ParentClass);
+				if (Pin->PinName.IsEqual(ContextParam) && Pin->LinkedTo.Num() > 0)
+				{
+					bConnectedWorldContext = true;
+					break;
+				}
+			}
+
+			if (!bConnectedWorldContext && ParentClass && !FBlueprintEditorUtils::ImplementsGetWorld(Blueprint) && !ParentClass->HasMetaDataHierarchical(FBlueprintMetadata::MD_ShowWorldContextPin))
+			{
+				MessageLog.Warning(*LOCTEXT("FunctionUnsafeInContext", "Function '@@' is unsafe to call from blueprints of class '@@' without connecting a World Context Object.").ToString(), this, ParentClass);
 			}
 		}
 
@@ -2588,7 +2591,9 @@ void UK2Node_CallFunction::ExpandNode(class FKismetCompilerContext& CompilerCont
 				}
 			};
 			FStructConnectHelper::Connect(DefaultToSelfMetaValue, this, BetterSelfPin, Schema, CompilerContext.MessageLog);
-			if (!Function->HasMetaData(FBlueprintMetadata::MD_CallableWithoutWorldContext))
+
+			// Attempt to connect, even if it's optional because of MD_CallableWithoutWorldContext
+			if (!WorldContextMetaValue.IsEmpty())
 			{
 				FStructConnectHelper::Connect(WorldContextMetaValue, this, BetterSelfPin, Schema, CompilerContext.MessageLog);
 			}
diff --git a/Engine/Source/Editor/BlueprintGraph/Private/K2Node_ConstructObjectFromClass.cpp b/Engine/Source/Editor/BlueprintGraph/Private/K2Node_ConstructObjectFromClass.cpp
--- a/Engine/Source/Editor/BlueprintGraph/Private/K2Node_ConstructObjectFromClass.cpp
+++ b/Engine/Source/Editor/BlueprintGraph/Private/K2Node_ConstructObjectFromClass.cpp
@@ -64,9 +64,7 @@ UClass* UK2Node_ConstructObjectFromClass::GetClassPinBaseClass() const
 
 bool UK2Node_ConstructObjectFromClass::UseWorldContext() const
 {
-	UBlueprint* BP = GetBlueprint();
-	const UClass* ParentClass = BP ? BP->ParentClass : nullptr;
-	return ParentClass ? ParentClass->HasMetaDataHierarchical(FBlueprintMetadata::MD_ShowWorldContextPin) != nullptr : false;
+	return FBlueprintEditorUtils::ShouldShowWorldContextPin(GetBlueprint());
 }
 
 void UK2Node_ConstructObjectFromClass::AllocateDefaultPins()
@@ -239,6 +237,10 @@ FString UK2Node_ConstructObjectFromClass::GetPinMetaData(FName InPinName, FName
 		{
 			MetaData = TEXT("false");
 		}
+		else if (InPinName == FK2Node_ConstructObjectFromClassHelper::GetWorldContextPinName() && InKey == FBlueprintMetadata::MD_HidePinAssetPicker)
+		{
+			MetaData = TEXT("true");
+		}
 	}
 
 	return MetaData;
diff --git a/Engine/Source/Editor/BlueprintGraph/Private/K2Node_GetSubsystem.cpp b/Engine/Source/Editor/BlueprintGraph/Private/K2Node_GetSubsystem.cpp
--- a/Engine/Source/Editor/BlueprintGraph/Private/K2Node_GetSubsystem.cpp
+++ b/Engine/Source/Editor/BlueprintGraph/Private/K2Node_GetSubsystem.cpp
@@ -52,6 +52,8 @@ class UEdGraph;
 //    UK2Node_GetSubsystem
 // ************************************************************************************
 
+static FName SubsystemWorldContextPinName("WorldContext");
+
 void UK2Node_GetSubsystem::Serialize( FArchive& Ar )
 {
 	Super::Serialize( Ar );
@@ -73,9 +75,9 @@ void UK2Node_GetSubsystem::Initialize( UClass* NodeClass )
 void UK2Node_GetSubsystem::AllocateDefaultPins()
 {
 	// If required add the world context pin
-	if (GetBlueprint()->ParentClass->HasMetaDataHierarchical(FBlueprintMetadata::MD_ShowWorldContextPin))
+	if (FBlueprintEditorUtils::ShouldShowWorldContextPin(GetBlueprint()))
 	{
-		CreatePin(EGPD_Input, UEdGraphSchema_K2::PC_Object, UObject::StaticClass(), TEXT("WorldContext"));
+		CreatePin(EGPD_Input, UEdGraphSchema_K2::PC_Object, UObject::StaticClass(), SubsystemWorldContextPinName);
 	}
 
 	// Add blueprint pin
@@ -90,6 +92,21 @@ void UK2Node_GetSubsystem::AllocateDefaultPins()
 	Super::AllocateDefaultPins();
 }
 
+FString UK2Node_GetSubsystem::GetPinMetaData(FName InPinName, FName InKey)
+{
+	FString MetaData = Super::GetPinMetaData(InPinName, InKey);
+
+	if (MetaData.IsEmpty())
+	{
+		if (InPinName == SubsystemWorldContextPinName && InKey == FBlueprintMetadata::MD_HidePinAssetPicker)
+		{
+			MetaData = TEXT("true");
+		}
+	}
+
+	return MetaData;
+}
+
 bool UK2Node_GetSubsystem::IsCompatibleWithGraph(const UEdGraph* TargetGraph) const
 {
 	UBlueprint* Blueprint = FBlueprintEditorUtils::FindBlueprintForGraph(TargetGraph);
@@ -328,7 +345,7 @@ FText UK2Node_GetSubsystem::GetTooltipText() const
 
 UEdGraphPin* UK2Node_GetSubsystem::GetWorldContextPin() const
 {
-	UEdGraphPin* Pin = FindPin(TEXT("WorldContext"));
+	UEdGraphPin* Pin = FindPin(SubsystemWorldContextPinName);
 	check(Pin == NULL || Pin->Direction == EGPD_Input);
 	return Pin;
 }
diff --git a/Engine/Source/Editor/BlueprintGraph/Private/K2Node_SpawnActor.cpp b/Engine/Source/Editor/BlueprintGraph/Private/K2Node_SpawnActor.cpp
--- a/Engine/Source/Editor/BlueprintGraph/Private/K2Node_SpawnActor.cpp
+++ b/Engine/Source/Editor/BlueprintGraph/Private/K2Node_SpawnActor.cpp
@@ -64,7 +64,7 @@ void UK2Node_SpawnActor::AllocateDefaultPins()
 	CreatePin(EGPD_Output, UEdGraphSchema_K2::PC_Exec, UEdGraphSchema_K2::PN_Then);
 
 	// If required add the world context pin
-	if (GetBlueprint()->ParentClass->HasMetaDataHierarchical(FBlueprintMetadata::MD_ShowWorldContextPin))
+	if (FBlueprintEditorUtils::ShouldShowWorldContextPin(GetBlueprint()))
 	{
 		CreatePin(EGPD_Input, UEdGraphSchema_K2::PC_Object, UObject::StaticClass(), FK2Node_SpawnActorHelper::WorldContextPinName);
 	}
diff --git a/Engine/Source/Editor/UnrealEd/Private/Kismet2/BlueprintEditorUtils.cpp b/Engine/Source/Editor/UnrealEd/Private/Kismet2/BlueprintEditorUtils.cpp
--- a/Engine/Source/Editor/UnrealEd/Private/Kismet2/BlueprintEditorUtils.cpp
+++ b/Engine/Source/Editor/UnrealEd/Private/Kismet2/BlueprintEditorUtils.cpp
@@ -4467,26 +4467,9 @@ void FBlueprintEditorUtils::GetHiddenPinsForFunction(UEdGraph const* Graph, UFun
 				const UEdGraphSchema_K2* K2Schema = GetDefault<UEdGraphSchema_K2>();
 				if(!K2Schema->IsStaticFunctionGraph(Graph))
 				{
-					bool bHasIntrinsicWorldContext = false;
-
-					UBlueprint const* CallingContext = FindBlueprintForGraph(Graph);
-					if (CallingContext && CallingContext->ParentClass)
-					{
-						UClass* NativeOwner = CallingContext->ParentClass;
-						while(NativeOwner && !NativeOwner->IsNative())
-						{
-							NativeOwner = NativeOwner->GetSuperClass();
-						}
-
-						if(NativeOwner)
-						{
-							bHasIntrinsicWorldContext = NativeOwner->GetDefaultObject()->ImplementsGetWorld();
-						}
-					}
-
 					// if the blueprint has world context that we can lookup with "self", 
 					// then we can hide this pin (and default it to self)
-					if (bHasIntrinsicWorldContext)
+					if (!ShouldShowWorldContextPin(FindBlueprintForGraph(Graph)))
 					{
 						HiddenPins.Add(*It.Value());
 					}
@@ -10150,6 +10133,25 @@ bool FBlueprintEditorUtils::ImplementsGetWorld(const UBlueprint* BP)
 	return false;
 }
 
+bool FBlueprintEditorUtils::ShouldShowWorldContextPin(const UBlueprint* BP)
+{
+	const UClass* ParentClass = BP ? BP->ParentClass : nullptr;
+
+	if (ParentClass && ParentClass->HasMetaDataHierarchical(FBlueprintMetadata::MD_ShowWorldContextPin) != nullptr)
+	{
+		// Always show even if not needed
+		return true;
+	}
+
+	if (BP && BP->BlueprintType == BPTYPE_FunctionLibrary)
+	{
+		// NOTE: This was the old behavior, setting it to false would hide the pin which could break existing blueprints
+		return true;
+	}
+
+	return !ImplementsGetWorld(BP);
+}
+
 struct FComponentInstancingDataUtils
 {
 	// Recursively gathers properties that differ from class/struct defaults, and fills out the cooked property list structure.
diff --git a/Engine/Source/Editor/UnrealEd/Public/Kismet2/BlueprintEditorUtils.h b/Engine/Source/Editor/UnrealEd/Public/Kismet2/BlueprintEditorUtils.h
--- a/Engine/Source/Editor/UnrealEd/Public/Kismet2/BlueprintEditorUtils.h
+++ b/Engine/Source/Editor/UnrealEd/Public/Kismet2/BlueprintEditorUtils.h
@@ -1966,6 +1966,11 @@ public:
 	 */
 	static UNREALED_API bool ImplementsGetWorld(const UBlueprint* BP);
 
+	/**
+	 * Returns true if this BP needs to show world context pins, because it does not implement GetWorld or has ShowWorldContextPin set
+	 */
+	static UNREALED_API bool ShouldShowWorldContextPin(const UBlueprint* BP);
+
 	/*
 	 * Checks a function for 'blueprint thread safety'. Note this is not strict thread safety and does not perform any
 	 * analysis on the function, it only inspects the metadata. Analysis is performed by the compiler for BP functions
