Protect against malformed FunctionVariableCache when gathering blueprint dependencies - cause unknown

#rnx

diff --git a/Engine/Source/Editor/BlueprintGraph/Private/K2Node_FunctionEntry.cpp b/Engine/Source/Editor/BlueprintGraph/Private/K2Node_FunctionEntry.cpp
--- a/Engine/Source/Editor/BlueprintGraph/Private/K2Node_FunctionEntry.cpp
+++ b/Engine/Source/Editor/BlueprintGraph/Private/K2Node_FunctionEntry.cpp
@@ -58,6 +58,11 @@
 
 #define LOCTEXT_NAMESPACE "K2Node_FunctionEntry"
 
+namespace UE::Private
+{
+extern KISMET_API bool GRecordingBlueprintDependencies;
+}
+
 //////////////////////////////////////////////////////////////////////////
 // FKCHandler_FunctionEntry
 
@@ -270,7 +275,7 @@ void UK2Node_FunctionEntry::Serialize(FArchive& Ar)
 
 	if (Ar.IsSaving())
 	{		
-		if (Ar.IsObjectReferenceCollector() || Ar.Tell() < 0)
+		if ((Ar.IsObjectReferenceCollector() || Ar.Tell() < 0) && !UE::Private::GRecordingBlueprintDependencies)
 		{
 			// If this is explicitly a reference collector, or it's a save with no backing archive, then we want to use the function variable cache if it exists
 			// It's not safe to regenerate the cache at this point as we could be in GIsSaving
diff --git a/Engine/Source/Editor/Kismet/Private/Blueprints/BlueprintDependencies.cpp b/Engine/Source/Editor/Kismet/Private/Blueprints/BlueprintDependencies.cpp
--- a/Engine/Source/Editor/Kismet/Private/Blueprints/BlueprintDependencies.cpp
+++ b/Engine/Source/Editor/Kismet/Private/Blueprints/BlueprintDependencies.cpp
@@ -44,6 +44,7 @@
 
 namespace UE::Private
 {
+KISMET_API bool GRecordingBlueprintDependencies = false;
 constexpr int32 BlueprintCookDependenciesVersion = 3;
 
 // gathering helpers:
@@ -661,6 +662,7 @@ void UE::Private::SaveBPCookDepenenciesImpl(FCbWriter& Writer, const FBlueprintD
 
 UE::Cook::FCookDependency UE::Private::BlueprintDependencies::RecordCookDependencies(const UBlueprint* BP)
 {
+	TGuardValue<bool> RecordDeps(GRecordingBlueprintDependencies, true);
 	if(!BP->ParentClass || 
 		!BP->GeneratedClass ||
 		!BP->GeneratedClass->IsChildOf(UObject::StaticClass()))
