Fix stale pointer in function variable cache during cook

[RN] Fixed crash when cooking blueprints with local object references set to objects in asset

diff --git a/Engine/Source/Editor/BlueprintGraph/Classes/K2Node_FunctionEntry.h b/Engine/Source/Editor/BlueprintGraph/Classes/K2Node_FunctionEntry.h
--- a/Engine/Source/Editor/BlueprintGraph/Classes/K2Node_FunctionEntry.h
+++ b/Engine/Source/Editor/BlueprintGraph/Classes/K2Node_FunctionEntry.h
@@ -55,6 +55,7 @@ class UK2Node_FunctionEntry : public UK2Node_FunctionTerminator
 
 	//~ Begin UObject Interface
 	virtual void Serialize(FArchive& Ar) override;
+	static void AddReferencedObjects(UObject* InThis, FReferenceCollector& Collector);
 	virtual void PreSave(FObjectPreSaveContext ObjectSaveContext) override;
 	virtual void PostLoad() override;
 	//~ End UObject Interface
diff --git a/Engine/Source/Editor/BlueprintGraph/Private/K2Node_FunctionEntry.cpp b/Engine/Source/Editor/BlueprintGraph/Private/K2Node_FunctionEntry.cpp
--- a/Engine/Source/Editor/BlueprintGraph/Private/K2Node_FunctionEntry.cpp
+++ b/Engine/Source/Editor/BlueprintGraph/Private/K2Node_FunctionEntry.cpp
@@ -58,11 +58,6 @@
 
 #define LOCTEXT_NAMESPACE "K2Node_FunctionEntry"
 
-namespace UE::Private
-{
-extern KISMET_API bool GRecordingBlueprintDependencies;
-}
-
 //////////////////////////////////////////////////////////////////////////
 // FKCHandler_FunctionEntry
 
@@ -275,7 +270,7 @@ void UK2Node_FunctionEntry::Serialize(FArchive& Ar)
 
 	if (Ar.IsSaving())
 	{		
-		if ((Ar.IsObjectReferenceCollector() || Ar.Tell() < 0) && !UE::Private::GRecordingBlueprintDependencies)
+		if (Ar.IsObjectReferenceCollector() || Ar.Tell() < 0)
 		{
 			// If this is explicitly a reference collector, or it's a save with no backing archive, then we want to use the function variable cache if it exists
 			// It's not safe to regenerate the cache at this point as we could be in GIsSaving
@@ -351,6 +346,17 @@ void UK2Node_FunctionEntry::Serialize(FArchive& Ar)
 	}
 }
 
+void UK2Node_FunctionEntry::AddReferencedObjects(UObject* InThis, FReferenceCollector& Collector)
+{
+	UK2Node_FunctionEntry* This = CastChecked<UK2Node_FunctionEntry>(InThis);
+	TSharedPtr<FStructOnScope> FunctionVariableCache = This->FunctionVariableCache;
+	if(FunctionVariableCache.IsValid())
+	{
+		Collector.AddPropertyReferences(FunctionVariableCache->GetStruct(), FunctionVariableCache->GetStructMemory());
+	}
+	return UK2Node_FunctionTerminator::AddReferencedObjects(InThis, Collector);
+}
+
 FText UK2Node_FunctionEntry::GetNodeTitle(ENodeTitleType::Type TitleType) const
 {
 	UEdGraph* Graph = GetGraph();
diff --git a/Engine/Source/Editor/Kismet/Private/Blueprints/BlueprintDependencies.cpp b/Engine/Source/Editor/Kismet/Private/Blueprints/BlueprintDependencies.cpp
--- a/Engine/Source/Editor/Kismet/Private/Blueprints/BlueprintDependencies.cpp
+++ b/Engine/Source/Editor/Kismet/Private/Blueprints/BlueprintDependencies.cpp
@@ -44,7 +44,6 @@
 
 namespace UE::Private
 {
-KISMET_API bool GRecordingBlueprintDependencies = false;
 constexpr int32 BlueprintCookDependenciesVersion = 3;
 
 // gathering helpers:
@@ -662,7 +661,6 @@ void UE::Private::SaveBPCookDepenenciesImpl(FCbWriter& Writer, const FBlueprintD
 
 UE::Cook::FCookDependency UE::Private::BlueprintDependencies::RecordCookDependencies(const UBlueprint* BP)
 {
-	TGuardValue<bool> RecordDeps(GRecordingBlueprintDependencies, true);
 	if(!BP->ParentClass || 
 		!BP->GeneratedClass ||
 		!BP->GeneratedClass->IsChildOf(UObject::StaticClass()))
