Allow selection of Object for BaseClassesToDisallowRecompilingDuringPlayInEditor, useful for disabling all compilation during PIE

#rnx

diff --git a/Engine/Source/Editor/BlueprintGraph/Public/BlueprintEditorSettings.h b/Engine/Source/Editor/BlueprintGraph/Public/BlueprintEditorSettings.h
--- a/Engine/Source/Editor/BlueprintGraph/Public/BlueprintEditorSettings.h
+++ b/Engine/Source/Editor/BlueprintGraph/Public/BlueprintEditorSettings.h
@@ -355,7 +355,7 @@ public:
 	 * Any blueprint deriving from one of these base classes will be disallowed to recompile during Play-in-Editor
 	 * (This setting exists both as an editor preference and project setting, and will be allowed if listed in either place) 
 	 */
-	UPROPERTY(EditAnywhere, config, Category=Play, meta=(AllowAbstract))
+	UPROPERTY(EditAnywhere, config, Category=Play, meta=(AllowAbstract, ShowObjectRootClass))
 	TArray<TSoftClassPtr<UObject>> BaseClassesToDisallowRecompilingDuringPlayInEditor;
 
 	/** Get allowed functions permissions list */
diff --git a/Engine/Source/Editor/UnrealEd/Public/Settings/BlueprintEditorProjectSettings.h b/Engine/Source/Editor/UnrealEd/Public/Settings/BlueprintEditorProjectSettings.h
--- a/Engine/Source/Editor/UnrealEd/Public/Settings/BlueprintEditorProjectSettings.h
+++ b/Engine/Source/Editor/UnrealEd/Public/Settings/BlueprintEditorProjectSettings.h
@@ -92,7 +92,7 @@ public:
 	 * Any blueprint deriving from one of these base classes will be disallowed to recompile during Play-in-Editor
 	 * (This setting exists both as an editor preference and project setting, and will be allowed if listed in either place) 
 	 */
-	UPROPERTY(EditAnywhere, config, Category=Play, meta=(AllowAbstract))
+	UPROPERTY(EditAnywhere, config, Category=Play, meta=(AllowAbstract, ShowObjectRootClass))
 	TArray<TSoftClassPtr<UObject>> BaseClassesToDisallowRecompilingDuringPlayInEditor;
 
 	/**
