Allow setting fields in structs that have a custom make node

[RN] Structs tagged as HasNativeMake can now still have their individual members set

diff --git a/Engine/Source/Editor/BlueprintGraph/Private/K2Node_SetFieldsInStruct.cpp b/Engine/Source/Editor/BlueprintGraph/Private/K2Node_SetFieldsInStruct.cpp
--- a/Engine/Source/Editor/BlueprintGraph/Private/K2Node_SetFieldsInStruct.cpp
+++ b/Engine/Source/Editor/BlueprintGraph/Private/K2Node_SetFieldsInStruct.cpp
@@ -90,9 +90,22 @@ public:
 		}
 	}
 
-	virtual void RegisterNets(FKismetFunctionContext& Context, UEdGraphNode* Node) override
+	virtual void RegisterNets(FKismetFunctionContext& Context, UEdGraphNode* InNode) override
 	{
-		FKCHandler_MakeStruct::RegisterNets(Context, Node);
+		UK2Node_SetFieldsInStruct* Node = CastChecked<UK2Node_SetFieldsInStruct>(InNode);
+		if (nullptr == Node->StructType)
+		{
+			CompilerContext.MessageLog.Error(*LOCTEXT("MakeStruct_UnknownStructure_Error", "Unknown structure to break for @@").ToString(), Node);
+			return;
+		}
+	
+		if (!UEdGraphSchema_K2::IsAllowableBlueprintVariableType(Node->StructType, Node->IsIntermediateNode()))
+		{
+			CompilerContext.MessageLog.Error(*LOCTEXT("MakeStruct_Error", "The structure @@ is not declared as BlueprintType - consider USTRUCT(BlueprintType) in C++ or resaving the asset").ToString(), Node);
+			return;
+		}
+
+		FKCHandler_MakeStruct::RegisterNetsImpl(Context, Node);
 
 		if (UEdGraphPin* ReturnPin = Node->FindPin(SetFieldsInStructHelper::StructOutPinName()))
 		{
diff --git a/Engine/Source/Editor/BlueprintGraph/Private/MakeStructHandler.cpp b/Engine/Source/Editor/BlueprintGraph/Private/MakeStructHandler.cpp
--- a/Engine/Source/Editor/BlueprintGraph/Private/MakeStructHandler.cpp
+++ b/Engine/Source/Editor/BlueprintGraph/Private/MakeStructHandler.cpp
@@ -68,13 +68,17 @@ void FKCHandler_MakeStruct::RegisterNets(FKismetFunctionContext& Context, UEdGra
 		CompilerContext.MessageLog.Error(*LOCTEXT("MakeStruct_UnknownStructure_Error", "Unknown structure to break for @@").ToString(), Node);
 		return;
 	}
-
+	
 	if (!UK2Node_MakeStruct::CanBeMade(Node->StructType, Node->IsIntermediateNode()))
 	{
 		CompilerContext.MessageLog.Error(*LOCTEXT("MakeStruct_Error", "The structure @@ is not a BlueprintType. ").ToString(), Node);
 		return;
 	}
-
+	return RegisterNetsImpl(Context, Node);
+}
+	
+void FKCHandler_MakeStruct::RegisterNetsImpl(FKismetFunctionContext& Context, UK2Node_MakeStruct* Node)
+{
 	FNodeHandlingFunctor::RegisterNets(Context, Node);
 
 	UEdGraphPin* OutputPin = FindStructPinChecked(Node);
diff --git a/Engine/Source/Editor/BlueprintGraph/Private/MakeStructHandler.h b/Engine/Source/Editor/BlueprintGraph/Private/MakeStructHandler.h
--- a/Engine/Source/Editor/BlueprintGraph/Private/MakeStructHandler.h
+++ b/Engine/Source/Editor/BlueprintGraph/Private/MakeStructHandler.h
@@ -6,6 +6,7 @@
 
 class FKismetCompilerContext;
 class UEdGraphNode;
+class UK2Node_MakeStruct;
 class UEdGraphPin;
 struct FBPTerminal;
 struct FKismetFunctionContext;
@@ -21,8 +22,9 @@ public:
 	virtual UEdGraphPin* FindStructPinChecked(UEdGraphNode* Node) const;
 
 	virtual void RegisterNets(FKismetFunctionContext& Context, UEdGraphNode* InNode) override;
-
 	virtual void RegisterNet(FKismetFunctionContext& Context, UEdGraphPin* Net) override;
+	void RegisterNetsImpl(FKismetFunctionContext& Context, UK2Node_MakeStruct* Node);
+
 	virtual FBPTerminal* RegisterLiteral(FKismetFunctionContext& Context, UEdGraphPin* Net) override;
 
 	virtual void Compile(FKismetFunctionContext& Context, UEdGraphNode* InNode) override;
