[Blueprint] Add a warning icon to a deprecated pin. Add the deprecated message to the tooltip.

diff --git a/Engine/Source/Editor/BlueprintGraph/Classes/EdGraphSchema_K2.h b/Engine/Source/Editor/BlueprintGraph/Classes/EdGraphSchema_K2.h
--- a/Engine/Source/Editor/BlueprintGraph/Classes/EdGraphSchema_K2.h
+++ b/Engine/Source/Editor/BlueprintGraph/Classes/EdGraphSchema_K2.h
@@ -930,6 +930,9 @@ public:
 
 	static UE_API bool IsPropertyExposedOnSpawn(const FProperty* Property);
 
+	/** returns whether the property is deprecated. */
+	static UE_API bool IsPropertyDeprecated(const FProperty* Property);
+
 	/**
 	 * Returns a list of parameters for the function that are specified as automatically emitting terms for unconnected ref parameters in the compiler (MD_AutoCreateRefTerm)
 	 *
diff --git a/Engine/Source/Editor/BlueprintGraph/Private/EdGraphSchema_K2.cpp b/Engine/Source/Editor/BlueprintGraph/Private/EdGraphSchema_K2.cpp
--- a/Engine/Source/Editor/BlueprintGraph/Private/EdGraphSchema_K2.cpp
+++ b/Engine/Source/Editor/BlueprintGraph/Private/EdGraphSchema_K2.cpp
@@ -1309,6 +1309,18 @@ bool UEdGraphSchema_K2::IsPropertyExposedOnSpawn(const FProperty* Property)
 	return false;
 }
 
+bool UEdGraphSchema_K2::IsPropertyDeprecated(const FProperty* Property)
+{
+	Property = FBlueprintEditorUtils::GetMostUpToDateProperty(Property);
+	if (Property)
+	{
+		const bool bMeta = Property->HasMetaData(FBlueprintMetadata::MD_DeprecatedProperty);
+		const bool bFlag = Property->HasAllPropertyFlags(CPF_Deprecated);
+		return bMeta || bFlag;
+	}
+	return false;
+}
+
 // if node is a get/set variable and the variable it refers to does not exist
 static bool IsUsingNonExistantVariable(const UEdGraphNode* InGraphNode, UBlueprint* OwnerBlueprint)
 {
@@ -3203,34 +3215,50 @@ void UEdGraphSchema_K2::ConstructBasicPinTooltip(const UEdGraphPin& Pin, const F
 		FFormatNamedArguments Args;
 		Args.Add(TEXT("PinType"), TypeToText(Pin.PinType));
 
+		if (Pin.bDeprecated && !Pin.DeprecationMessage.IsEmpty())
+		{
+			Args.Add(TEXT("DeprecationWarning"), FText::FromString(Pin.DeprecationMessage));
+			Args.Add(TEXT("LineFeed1"), FText::FromString(TEXT("\n")));
+		}
+		else if (Pin.bDeprecated)
+		{
+			Args.Add(TEXT("DeprecationWarning"), LOCTEXT("Deprecated", "Deprecated"));
+			Args.Add(TEXT("LineFeed1"), FText::FromString(TEXT("\n")));
+		}
+		else
+		{
+			Args.Add(TEXT("DeprecationWarning"), FText::GetEmpty());
+			Args.Add(TEXT("LineFeed1"), FText::GetEmpty());
+		}
+
 		if (UEdGraphNode* PinNode = Pin.GetOwningNode())
 		{
 			UEdGraphSchema_K2 const* const K2Schema = Cast<const UEdGraphSchema_K2>(PinNode->GetSchema());
 			if (ensure(K2Schema != NULL)) // ensure that this node belongs to this schema
 			{
 				Args.Add(TEXT("DisplayName"), GetPinDisplayName(&Pin));
-				Args.Add(TEXT("LineFeed1"), FText::FromString(TEXT("\n")));
+				Args.Add(TEXT("LineFeed2"), FText::FromString(TEXT("\n")));
 			}
 		}
 		else
 		{
 				Args.Add(TEXT("DisplayName"), FText::GetEmpty());
-				Args.Add(TEXT("LineFeed1"), FText::GetEmpty());
+				Args.Add(TEXT("LineFeed2"), FText::GetEmpty());
 		}
 
 
 		if (!PinDescription.IsEmpty())
 		{
 			Args.Add(TEXT("Description"), PinDescription);
-			Args.Add(TEXT("LineFeed2"), FText::FromString(TEXT("\n\n")));
+			Args.Add(TEXT("LineFeed3"), FText::FromString(TEXT("\n\n")));
 		}
 		else
 		{
 			Args.Add(TEXT("Description"), FText::GetEmpty());
-			Args.Add(TEXT("LineFeed2"), FText::GetEmpty());
+			Args.Add(TEXT("LineFeed3"), FText::GetEmpty());
 		}
 	
-		TooltipOut = FText::Format(LOCTEXT("PinTooltip", "{DisplayName}{LineFeed1}{PinType}{LineFeed2}{Description}"), Args).ToString(); 
+		TooltipOut = FText::Format(LOCTEXT("PinTooltip", "{DeprecationWarning}{LineFeed1}{DisplayName}{LineFeed2}{PinType}{LineFeed3}{Description}"), Args).ToString(); 
 	}
 }
 
@@ -7368,6 +7396,8 @@ void UEdGraphSchema_K2::SplitPin(UEdGraphPin* Pin, const bool bNotify) const
 
 			SubPin->DefaultValue = MoveTemp(ProtoPin->DefaultValue);
 			SubPin->AutogeneratedDefaultValue = MoveTemp(ProtoPin->AutogeneratedDefaultValue);
+			SubPin->bDeprecated = ProtoPin->bDeprecated;
+			SubPin->DeprecationMessage = ProtoPin->DeprecationMessage;
 
 			SubPin->ParentPin = Pin;
 
diff --git a/Engine/Source/Editor/BlueprintGraph/Private/K2Node.cpp b/Engine/Source/Editor/BlueprintGraph/Private/K2Node.cpp
--- a/Engine/Source/Editor/BlueprintGraph/Private/K2Node.cpp
+++ b/Engine/Source/Editor/BlueprintGraph/Private/K2Node.cpp
@@ -1845,6 +1845,11 @@ void FOptionalPinManager::CreateVisiblePins(TArray<FOptionalPinFromProperty>& Pr
 							NewPin->PinFriendlyName = PinFriendlyName;
 							NewPin->bNotConnectable = !PropertyEntry.bIsSetValuePinVisible;
 							NewPin->bDefaultValueIsIgnored = !PropertyEntry.bIsSetValuePinVisible;
+							if (UEdGraphSchema_K2::IsPropertyDeprecated(OuterProperty))
+							{
+								NewPin->bDeprecated = true;
+								NewPin->DeprecationMessage = OuterProperty->GetMetaData(FBlueprintMetadata::MD_DeprecationMessage);
+							}
 							Schema->ConstructBasicPinTooltip(*NewPin, PropertyEntry.PropertyTooltip, NewPin->PinToolTip);
 
 							// Allow the derived class to customize the created pin
@@ -1891,6 +1896,11 @@ void FOptionalPinManager::CreateVisiblePins(TArray<FOptionalPinFromProperty>& Pr
 						NewPin->PinFriendlyName = FText::AsCultureInvariant(PropertyEntry.PropertyFriendlyName.IsEmpty() ? PinName.ToString() : PropertyEntry.PropertyFriendlyName);
 						NewPin->bNotConnectable = !PropertyEntry.bIsSetValuePinVisible;
 						NewPin->bDefaultValueIsIgnored = !PropertyEntry.bIsSetValuePinVisible;
+						if (UEdGraphSchema_K2::IsPropertyDeprecated(OuterProperty))
+						{
+							NewPin->bDeprecated = true;
+							NewPin->DeprecationMessage = OuterProperty->GetMetaData(FBlueprintMetadata::MD_DeprecationMessage);
+						}
 						Schema->ConstructBasicPinTooltip(*NewPin, PropertyEntry.PropertyTooltip, NewPin->PinToolTip);
 
 						// Allow the derived class to customize the created pin
diff --git a/Engine/Source/Editor/BlueprintGraph/Private/K2Node_ConstructObjectFromClass.cpp b/Engine/Source/Editor/BlueprintGraph/Private/K2Node_ConstructObjectFromClass.cpp
--- a/Engine/Source/Editor/BlueprintGraph/Private/K2Node_ConstructObjectFromClass.cpp
+++ b/Engine/Source/Editor/BlueprintGraph/Private/K2Node_ConstructObjectFromClass.cpp
@@ -157,6 +157,12 @@ void UK2Node_ConstructObjectFromClass::CreatePinsForClass(UClass* InClass, TArra
 					}
 				}
 
+				if (UEdGraphSchema_K2::IsPropertyDeprecated(Property))
+				{
+					Pin->bDeprecated = true;
+					Pin->DeprecationMessage = Property->GetMetaData(FBlueprintMetadata::MD_DeprecationMessage);
+				}
+
 				if (ClassDefaultObject && K2Schema->PinDefaultValueIsEditable(*Pin))
 				{
 					FString DefaultValueAsString;
diff --git a/Engine/Source/Editor/GraphEditor/Private/SGraphPin.cpp b/Engine/Source/Editor/GraphEditor/Private/SGraphPin.cpp
--- a/Engine/Source/Editor/GraphEditor/Private/SGraphPin.cpp
+++ b/Engine/Source/Editor/GraphEditor/Private/SGraphPin.cpp
@@ -250,6 +250,15 @@ void SGraphPin::Construct(const FArguments& InArgs, UEdGraphPin* InPin)
 		SNew(SWrapBox)
 		.PreferredSize(150.f);
 
+	TSharedPtr<SWidget> PinDeprecationIndicator;
+	if (InPin->bDeprecated)
+	{
+		static const FName DeprecationWarning("Icons.WarningWithColor");
+		PinDeprecationIndicator = SNew(SImage)
+			.Image(FAppStyle::GetBrush(DeprecationWarning))
+			.ToolTipText(FText::FromString(InPin->DeprecationMessage));
+	}
+
 	if (!bIsInput)
 	{
 		LabelAndValue->AddSlot()
@@ -258,6 +267,15 @@ void SGraphPin::Construct(const FArguments& InArgs, UEdGraphPin* InPin)
 				PinStatusIndicator
 			];
 
+		if (PinDeprecationIndicator)
+		{
+			LabelAndValue->AddSlot()
+				.VAlign(VAlign_Center)
+				[
+					PinDeprecationIndicator.ToSharedRef()
+				];
+		}
+
 		LabelAndValue->AddSlot()
 			.VAlign(VAlign_Center)
 			[
@@ -266,6 +284,15 @@ void SGraphPin::Construct(const FArguments& InArgs, UEdGraphPin* InPin)
 	}
 	else
 	{
+		if (PinDeprecationIndicator)
+		{
+			LabelAndValue->AddSlot()
+				.VAlign(VAlign_Center)
+				[
+					PinDeprecationIndicator.ToSharedRef()
+				];
+		}
+
 		LabelAndValue->AddSlot()
 			.VAlign(VAlign_Center)
 			[
diff --git a/Engine/Source/Runtime/Engine/Classes/EdGraph/EdGraphPin.h b/Engine/Source/Runtime/Engine/Classes/EdGraph/EdGraphPin.h
--- a/Engine/Source/Runtime/Engine/Classes/EdGraph/EdGraphPin.h
+++ b/Engine/Source/Runtime/Engine/Classes/EdGraph/EdGraphPin.h
@@ -321,6 +321,9 @@ public:
 	/** If true, the pin may be hidden by user. */
 	uint8 bAdvancedView:1;
 
+	/** If true, the pin is deprecated. */
+	uint8 bDeprecated:1;
+
 	/** If true, the pin is displayed as ref. This is transient. */
 	uint8 bDisplayAsMutableRef:1;
 
@@ -351,6 +354,9 @@ public:
 
 #if WITH_EDITORONLY_DATA
 	FText PinFriendlyName;
+
+	/** The deprecated message of the pin. Used when the pin is deprecated. */
+	FString DeprecationMessage;
 #endif
 
 	/** The tool-tip describing this pin's purpose */
@@ -386,7 +392,6 @@ public:
 #if WITH_EDITORONLY_DATA
 	/** Pin name could be changed, so whenever possible it's good to have a persistent GUID identifying Pin to reconstruct Node seamlessly */
 	FGuid PersistentGuid;
-
 #endif
 
 public:
@@ -497,6 +502,7 @@ public:
 		SourceIndex = INDEX_NONE;
 #if WITH_EDITORONLY_DATA
 		PinFriendlyName = FText::GetEmpty();
+		DeprecationMessage.Reset();
 #endif // WITH_EDITORONLY_DATA
 		AutogeneratedDefaultValue.Reset();
 		ResetDefaultValue();
@@ -506,6 +512,7 @@ public:
 		bNotConnectable = false;
 		bDefaultValueIsReadOnly = false;
 		bDefaultValueIsIgnored = false;
+		bDeprecated = false;
 		bOrphanedPin = false;
 		bSavePinIfOrphaned = true;
 		bAllowFriendlyName = true;
diff --git a/Engine/Source/Runtime/Engine/Private/EdGraph/EdGraphPin.cpp b/Engine/Source/Runtime/Engine/Private/EdGraph/EdGraphPin.cpp
--- a/Engine/Source/Runtime/Engine/Private/EdGraph/EdGraphPin.cpp
+++ b/Engine/Source/Runtime/Engine/Private/EdGraph/EdGraphPin.cpp
@@ -1426,6 +1426,7 @@ UEdGraphPin::UEdGraphPin(UEdGraphNode* InOwningNode, const FGuid& PinIdGuid)
 	, bDefaultValueIsReadOnly(false)
 	, bDefaultValueIsIgnored(false)
 	, bAdvancedView(false)
+	, bDeprecated(false)
 	, bDisplayAsMutableRef(false)
 	, bAllowFriendlyName(true)
 	, bOrphanedPin(false)
