Fix SetFieldsInStruct trying to deprecate the node when MakeStruct is deprecated

diff --git a/Engine/Source/Editor/BlueprintGraph/Classes/K2Node_SetFieldsInStruct.h b/Engine/Source/Editor/BlueprintGraph/Classes/K2Node_SetFieldsInStruct.h
--- a/Engine/Source/Editor/BlueprintGraph/Classes/K2Node_SetFieldsInStruct.h
+++ b/Engine/Source/Editor/BlueprintGraph/Classes/K2Node_SetFieldsInStruct.h
@@ -44,6 +44,7 @@ class UK2Node_SetFieldsInStruct : public UK2Node_MakeStruct
 	//~ Begin K2Node Interface
 	virtual bool IsNodePure() const override { return false; }
 	virtual class FNodeHandlingFunctor* CreateNodeHandler(class FKismetCompilerContext& CompilerContext) const override;
+	virtual void ConvertDeprecatedNode(UEdGraph* Graph, bool bOnlySafeChanges) override {};
 	//~ End K2Node Interface
 
 	BLUEPRINTGRAPH_API static bool ShowCustomPinActions(const UEdGraphPin* InGraphPin, bool bIgnorePinsNum);
