Add more ways to detect a local player subsystem when using the node which does not have an explicit player controller target pin.

- Check for the instigating controller on an actor
- For actor components, get the owning actor and check that instigating controller.

Also add a tooltip to these local player subsystem nodes to make it clearer when they will work at runtime and when they wont.

[RN] Any "Local Player Subsystem" nodes which did not take an explicit player controller pin reference will now work in more scenarios, instead of being a "no-op".

diff --git a/Engine/Source/Editor/BlueprintGraph/Private/K2Node_GetSubsystem.cpp b/Engine/Source/Editor/BlueprintGraph/Private/K2Node_GetSubsystem.cpp
--- a/Engine/Source/Editor/BlueprintGraph/Private/K2Node_GetSubsystem.cpp
+++ b/Engine/Source/Editor/BlueprintGraph/Private/K2Node_GetSubsystem.cpp
@@ -320,6 +320,7 @@ FText UK2Node_GetSubsystem::GetTooltipText() const
 	if (CustomClass)
 	{
 		FText SubsystemTypeText;
+		FText SubsystemNoteText;
 		if (CustomClass->IsChildOf<UGameInstanceSubsystem>())
 		{
 			SubsystemTypeText =  NSLOCTEXT("K2Node", "GetSubsystem_GameInstanceSubsystemTooltip", "GameInstance Subsystem");
@@ -335,11 +336,13 @@ FText UK2Node_GetSubsystem::GetTooltipText() const
 		else
 		{
 			SubsystemTypeText = NSLOCTEXT("K2Node", "GetSubsystem_LocalPlayerSubsystemTooltip", "LocalPlayer Subsystem");
+			SubsystemNoteText = NSLOCTEXT("K2Node", "GetSubsystem_LocalPlayerSubsystemTooltip", "\n\nNote: This will only be valid if there is an available local player to this object.");
 		}
-		return FText::FormatNamed(NSLOCTEXT("K2Node", "GetSubsystem_TooltipFormat", "Get {ClassName} ({SubsystemType})\n\n{ClassTooltip}"),
+		return FText::FormatNamed(NSLOCTEXT("K2Node", "GetSubsystem_TooltipFormat", "Get {ClassName} ({SubsystemType})\n\n{ClassTooltip}{SubsystemNote}"),
 			TEXT("ClassName"), CustomClass->GetDisplayNameText(),
 			TEXT("SubsystemType"), SubsystemTypeText,
-			TEXT("ClassTooltip"), CustomClass->GetToolTipText(/*bShortTooltip=*/ true));
+			TEXT("ClassTooltip"), CustomClass->GetToolTipText(/*bShortTooltip=*/ true),
+			TEXT("SubsystemNote"), SubsystemNoteText);
 	}
 
 	return NSLOCTEXT("K2Node", "GetSubsystem_InvalidSubsystemTypeTooltip", "Invalid Subsystem Type");
@@ -511,7 +514,10 @@ FText UK2Node_GetSubsystemFromPC::GetTooltipText() const
 {
 	if (CustomClass)
 	{
-		return FText::FormatNamed(NSLOCTEXT("K2Node", "GetSubsystemFromPC_TooltipFormat", "Get {ClassName} from Player Controller"), TEXT("ClassName"), CustomClass->GetDisplayNameText());
+		return
+			FText::FormatNamed(NSLOCTEXT("K2Node", "GetSubsystemFromPC_TooltipFormat", "Get {ClassName} from Player Controller\n\n{ClassTooltip}"),
+			TEXT("ClassName"), CustomClass->GetDisplayNameText(),
+			TEXT("ClassTooltip"), CustomClass->GetToolTipText(/*bShortTooltip=*/ true));
 	}
 
 	return NSLOCTEXT("K2Node", "GetSubsystemFromPC_InvalidSubsystemTypeTooltip", "Invalid Subsystem Type");
diff --git a/Engine/Source/Runtime/Engine/Private/Subsystems/SubsystemBlueprintLibrary.cpp b/Engine/Source/Runtime/Engine/Private/Subsystems/SubsystemBlueprintLibrary.cpp
--- a/Engine/Source/Runtime/Engine/Private/Subsystems/SubsystemBlueprintLibrary.cpp
+++ b/Engine/Source/Runtime/Engine/Private/Subsystems/SubsystemBlueprintLibrary.cpp
@@ -2,10 +2,13 @@
 
 #include "Subsystems/SubsystemBlueprintLibrary.h"
 
-#include "Engine/Engine.h"
-#include "Engine/GameInstance.h"
 #include "AudioDevice.h"
 #include "Blueprint/UserWidget.h"
+#include "Components/ActorComponent.h"
+#include "Engine/Engine.h"
+#include "Engine/GameInstance.h"
+#include "GameFramework/Actor.h"
+#include "GameFramework/Pawn.h"
 #include "Subsystems/EngineSubsystem.h"
 
 #include UE_INLINE_GENERATED_CPP_BY_NAME(SubsystemBlueprintLibrary)
@@ -35,10 +38,29 @@
 	{
 		LocalPlayer = UserWidget->GetOwningLocalPlayer();
 	}
-	else if (APlayerController* PlayerController = Cast<APlayerController>(ContextObject))
+	else if (const APlayerController* PlayerController = Cast<APlayerController>(ContextObject))
 	{
 		LocalPlayer = Cast<ULocalPlayer>(PlayerController->Player);
 	}
+	else if (const APawn* Pawn = Cast<APawn>(ContextObject))
+	{
+		const APlayerController* PC = Cast<APlayerController>(Pawn->GetController());
+		
+		LocalPlayer = PC ? PC->GetLocalPlayer() : nullptr;
+	}
+	else if (const AActor* Actor = Cast<AActor>(ContextObject))
+	{
+		const APlayerController* PC = Cast<APlayerController>(Actor->GetInstigatorController());
+		
+		LocalPlayer = PC ? PC->GetLocalPlayer() : nullptr;
+	}
+	else if (const UActorComponent* ActorComponent = Cast<UActorComponent>(ContextObject))
+	{
+		const AActor* OwningActor = ActorComponent->GetOwner<AActor>();
+		const APlayerController* PC = Cast<APlayerController>(OwningActor->GetInstigatorController());
+		
+		LocalPlayer = PC ? PC->GetLocalPlayer() : nullptr;
+	}
 	else
 	{
 		LocalPlayer = Cast<ULocalPlayer>(ContextObject);
