Fix up for FText local namespace duplication

#rnx

diff --git a/Engine/Source/Editor/BlueprintGraph/Private/K2Node_GetSubsystem.cpp b/Engine/Source/Editor/BlueprintGraph/Private/K2Node_GetSubsystem.cpp
--- a/Engine/Source/Editor/BlueprintGraph/Private/K2Node_GetSubsystem.cpp
+++ b/Engine/Source/Editor/BlueprintGraph/Private/K2Node_GetSubsystem.cpp
@@ -335,8 +335,8 @@ FText UK2Node_GetSubsystem::GetTooltipText() const
 		}
 		else
 		{
-			SubsystemTypeText = NSLOCTEXT("K2Node", "GetSubsystem_LocalPlayerSubsystemTooltip", "LocalPlayer Subsystem");
-			SubsystemNoteText = NSLOCTEXT("K2Node", "GetSubsystem_LocalPlayerSubsystemTooltip", "\n\nNote: This will only be valid if there is an available local player to this object.");
+			SubsystemTypeText = NSLOCTEXT("K2Node", "GetSubsystem_LocalPlayerSubsystemTooltip_Label", "LocalPlayer Subsystem");
+			SubsystemNoteText = NSLOCTEXT("K2Node", "GetSubsystem_LocalPlayerSubsystemTooltip_Details", "\n\nNote: This will only be valid if there is an available local player to this object.");
 		}
 		return FText::FormatNamed(NSLOCTEXT("K2Node", "GetSubsystem_TooltipFormat", "Get {ClassName} ({SubsystemType})\n\n{ClassTooltip}{SubsystemNote}"),
 			TEXT("ClassName"), CustomClass->GetDisplayNameText(),
