UE Blueprints: Adding a new "Async Load Assets" node that allows a user to async load a list of soft object references

diff --git a/Engine/Source/Editor/BlueprintGraph/Classes/K2Node_LoadAsset.h b/Engine/Source/Editor/BlueprintGraph/Classes/K2Node_LoadAsset.h
--- a/Engine/Source/Editor/BlueprintGraph/Classes/K2Node_LoadAsset.h
+++ b/Engine/Source/Editor/BlueprintGraph/Classes/K2Node_LoadAsset.h
@@ -73,4 +73,22 @@ protected:
 	virtual const FName& GetOutputPinName() const override;
 };
 
+UCLASS(MinimalAPI)
+class UK2Node_LoadAssets : public UK2Node_LoadAsset
+{
+	GENERATED_BODY()
+public:
+	// UEdGraphNode interface
+	virtual void AllocateDefaultPins() override;
+	virtual FText GetTooltipText() const override;
+	virtual FText GetNodeTitle(ENodeTitleType::Type TitleType) const override;
+	// End of UEdGraphNode interface
+
+protected:
+	virtual FName NativeFunctionName() const override;
+
+	virtual const FName& GetInputPinName() const override;
+	virtual const FName& GetOutputPinName() const override;
+};
+
 #undef UE_API
diff --git a/Engine/Source/Editor/BlueprintGraph/Private/K2Node_LoadAsset.cpp b/Engine/Source/Editor/BlueprintGraph/Private/K2Node_LoadAsset.cpp
--- a/Engine/Source/Editor/BlueprintGraph/Private/K2Node_LoadAsset.cpp
+++ b/Engine/Source/Editor/BlueprintGraph/Private/K2Node_LoadAsset.cpp
@@ -93,7 +93,7 @@ void UK2Node_LoadAsset::ExpandNode(class FKismetCompilerContext& CompilerContext
 	// Create LoadAsset function call
 	UK2Node_CallFunction* CallLoadAssetNode = CompilerContext.SpawnIntermediateNode<UK2Node_CallFunction>(this, SourceGraph);
 	CallLoadAssetNode->FunctionReference.SetExternalMember(NativeFunctionName(), UKismetSystemLibrary::StaticClass());
-	CallLoadAssetNode->AllocateDefaultPins(); 
+	CallLoadAssetNode->AllocateDefaultPins();
 
 	// connect load to first sequence pin
 	{
@@ -110,7 +110,13 @@ void UK2Node_LoadAsset::ExpandNode(class FKismetCompilerContext& CompilerContext
 	}
 
 	// Create Local Variable
-	UK2Node_TemporaryVariable* TempVarOutput = CompilerContext.SpawnInternalVariable(this, GetOutputCategory(), NAME_None, UObject::StaticClass());
+	UEdGraphPin* AssetPin = FindPin(GetInputPinName());
+	EPinContainerType PinContainerType = EPinContainerType::None;
+	if (AssetPin)
+	{
+		PinContainerType = AssetPin->PinType.ContainerType;
+	}
+	UK2Node_TemporaryVariable* TempVarOutput = CompilerContext.SpawnInternalVariable(this, GetOutputCategory(), NAME_None, UObject::StaticClass(), PinContainerType);
 
 	// Create assign node
 	UK2Node_AssignmentStatement* AssignNode = CompilerContext.SpawnIntermediateNode<UK2Node_AssignmentStatement>(this, SourceGraph);
@@ -130,13 +136,11 @@ void UK2Node_LoadAsset::ExpandNode(class FKismetCompilerContext& CompilerContext
 		bIsErrorFree &= LoadedObjectVariablePin && OutputObjectPinPin && CompilerContext.MovePinLinksToIntermediate(*OutputObjectPinPin, *LoadedObjectVariablePin).CanSafeConnect();
 	}
 
-	
 	// connect to asset
 	UEdGraphPin* CallFunctionAssetPin = CallLoadAssetNode->FindPin(GetInputPinName());
 	{
-		UEdGraphPin* AssetPin = FindPin(GetInputPinName());
-		ensure(CallFunctionAssetPin);
 
+		ensure(CallFunctionAssetPin);
 		if (AssetPin && CallFunctionAssetPin)
 		{
 			if (AssetPin->LinkedTo.Num() > 0)
@@ -295,13 +299,11 @@ FName UK2Node_LoadAsset::NativeFunctionName() const
 }
 
 // UK2Node_LoadAssetClass
-
 FName UK2Node_LoadAssetClass::NativeFunctionName() const
 {
 	return GET_FUNCTION_NAME_CHECKED(UKismetSystemLibrary, LoadAssetClass);
 }
 
-
 FText UK2Node_LoadAssetClass::GetTooltipText() const
 {
 	return FText(LOCTEXT("UK2Node_LoadAssetClassGetTooltipText", "Asynchronously loads a Soft Class Reference and returns class of the correct type if the load succeeds"));
@@ -334,5 +336,45 @@ const FName& UK2Node_LoadAssetClass::GetOutputPinName() const
 	return OutputObjectPinName;
 }
 
+// UK2Node_LoadAssets
+void UK2Node_LoadAssets::AllocateDefaultPins()
+{
+	CreatePin(EGPD_Input, UEdGraphSchema_K2::PC_Exec, UEdGraphSchema_K2::PN_Execute);
+	// The immediate continue pin
+	CreatePin(EGPD_Output, UEdGraphSchema_K2::PC_Exec, UEdGraphSchema_K2::PN_Then);
+	// The delayed completed pin, this used to be called Then
+	CreatePin(EGPD_Output, UEdGraphSchema_K2::PC_Exec, UEdGraphSchema_K2::PN_Completed);
+	FCreatePinParams ArrayParams;
+	ArrayParams.ContainerType = EPinContainerType::Array;
+	CreatePin(EGPD_Input, GetInputCategory(), UObject::StaticClass(), GetInputPinName(), ArrayParams);
+	CreatePin(EGPD_Output, GetOutputCategory(), UObject::StaticClass(), GetOutputPinName(), ArrayParams);
+}
+
+FText UK2Node_LoadAssets::GetTooltipText() const
+{
+	return FText(LOCTEXT("UK2Node_LoadAssetsGetTooltipText", "Asynchronously loads a list Soft Object References and returns the objects of the correct type if the load succeeds"));
+}
+
+FText UK2Node_LoadAssets::GetNodeTitle(ENodeTitleType::Type TitleType) const
+{
+	return FText(LOCTEXT("UK2Node_LoadAssetClassGetNodeTitle", "Async Load Assets"));
+}
+
+FName UK2Node_LoadAssets::NativeFunctionName() const
+{
+	return GET_FUNCTION_NAME_CHECKED(UKismetSystemLibrary, LoadAssets);
+}
+
+const FName& UK2Node_LoadAssets::GetInputPinName() const
+{
+	static const FName InputAssetPinName("Assets");
+	return InputAssetPinName;
+}
+
+const FName& UK2Node_LoadAssets::GetOutputPinName() const
+{
+	static const FName OutputObjectPinName("Objects");
+	return OutputObjectPinName;
+}
 
-#undef LOCTEXT_NAMESPACE
+#undef LOCTEXT_NAMESPACE
\ No newline at end of file
diff --git a/Engine/Source/Runtime/Engine/Classes/Kismet/KismetSystemLibrary.h b/Engine/Source/Runtime/Engine/Classes/Kismet/KismetSystemLibrary.h
--- a/Engine/Source/Runtime/Engine/Classes/Kismet/KismetSystemLibrary.h
+++ b/Engine/Source/Runtime/Engine/Classes/Kismet/KismetSystemLibrary.h
@@ -449,6 +449,11 @@ class UKismetSystemLibrary : public UBlueprintFunctionLibrary
 	UFUNCTION(BlueprintCallable, meta = (Latent, LatentInfo = "LatentInfo", WorldContext = "WorldContextObject", BlueprintInternalUseOnly = "true"), Category = "Utilities")
 	static ENGINE_API void LoadAssetClass(const UObject* WorldContextObject, TSoftClassPtr<UObject> AssetClass, FOnAssetClassLoaded OnLoaded, FLatentActionInfo LatentInfo);
 
+	DECLARE_DYNAMIC_DELEGATE_OneParam(FOnAssetsLoaded, const TArray<class UObject*>&, Loaded);
+
+	UFUNCTION(BlueprintCallable, meta = (Latent, LatentInfo = "LatentInfo", WorldContext = "WorldContextObject", BlueprintInternalUseOnly = "true"), Category = "Utilities")
+	static ENGINE_API void LoadAssets(const UObject* WorldContextObject, TArray<TSoftObjectPtr<UObject>> Assets, FOnAssetsLoaded OnLoaded, FLatentActionInfo LatentInfo);
+
 	/**
 	 * Creates a literal integer
 	 * @param	Value	value to set the integer to
diff --git a/Engine/Source/Runtime/Engine/Private/KismetSystemLibrary.cpp b/Engine/Source/Runtime/Engine/Private/KismetSystemLibrary.cpp
--- a/Engine/Source/Runtime/Engine/Private/KismetSystemLibrary.cpp
+++ b/Engine/Source/Runtime/Engine/Private/KismetSystemLibrary.cpp
@@ -45,6 +45,7 @@
 #include "Commandlets/Commandlet.h"
 #include "UObject/PropertyAccessUtil.h"
 #include "UObject/TextProperty.h"
+#include "UObject/SoftObjectPath.h"
 
 #include UE_INLINE_GENERATED_CPP_BY_NAME(KismetSystemLibrary)
 
@@ -3162,6 +3163,98 @@ void UKismetSystemLibrary::LoadAssetClass(const UObject* WorldContextObject, TSo
 	}
 }
 
+void UKismetSystemLibrary::LoadAssets(const UObject* WorldContextObject, TArray<TSoftObjectPtr<UObject>> Assets, FOnAssetsLoaded OnLoaded, FLatentActionInfo LatentInfo)
+{
+	struct FLoadAssetsAction : public FPendingLatentAction
+	{
+	public:
+
+		TArray<FSoftObjectPath> SoftObjectPaths;
+		FStreamableManager StreamableManager;
+		TSharedPtr<FStreamableHandle> Handle;
+		FName ExecutionFunction;
+		int32 OutputLink;
+		FWeakObjectPtr CallbackTarget;
+		UKismetSystemLibrary::FOnAssetsLoaded OnLoadedCallback;
+
+		FLoadAssetsAction(const TArray<FSoftObjectPath>& InSoftObjectPaths, UKismetSystemLibrary::FOnAssetsLoaded InOnLoadedCallback, const FLatentActionInfo& InLatentInfo)
+			: SoftObjectPaths(InSoftObjectPaths)
+			, ExecutionFunction(InLatentInfo.ExecutionFunction)
+			, OutputLink(InLatentInfo.Linkage)
+			, CallbackTarget(InLatentInfo.CallbackTarget)
+			, OnLoadedCallback(InOnLoadedCallback)
+
+		{
+			Handle = StreamableManager.RequestAsyncLoad(SoftObjectPaths);
+		}
+
+		virtual ~FLoadAssetsAction()
+		{
+			if (Handle.IsValid())
+			{
+				Handle->ReleaseHandle();
+			}
+		}
+
+		virtual void UpdateOperation(FLatentResponse& Response) override
+		{
+			const bool bLoaded = !Handle.IsValid() || Handle->HasLoadCompleted() || Handle->WasCanceled();
+			if (bLoaded)
+			{
+				OnLoaded();
+			}
+			Response.FinishAndTriggerIf(bLoaded, ExecutionFunction, OutputLink, CallbackTarget);
+		}
+
+#if WITH_EDITOR
+		virtual FString GetDescription() const override
+		{
+			FString AggregatedSoftObjectPaths;
+
+			for (auto Path = SoftObjectPaths.CreateConstIterator(); Path; ++Path)
+			{
+				AggregatedSoftObjectPaths.Append(Path->ToString());
+
+				if (*Path != SoftObjectPaths.Last())
+				{
+					AggregatedSoftObjectPaths.Append("\n");
+				}
+			}
+
+			return FString::Printf(TEXT("Load Assets Action: %s"), *AggregatedSoftObjectPaths);
+		}
+#endif
+
+		void OnLoaded() 
+		{
+			TArray<UObject*> LoadedObjects;
+
+			for (FSoftObjectPath& Path : SoftObjectPaths)
+			{
+				LoadedObjects.Add(Path.ResolveObject());
+			}
+
+			OnLoadedCallback.ExecuteIfBound(LoadedObjects);
+		}
+	};
+
+	if (UWorld* World = GEngine->GetWorldFromContextObject(WorldContextObject, EGetWorldErrorMode::LogAndReturnNull))
+	{
+		FLatentActionManager& LatentManager = World->GetLatentActionManager();
+
+		TArray<FSoftObjectPath> SoftObjectPathsToLoad;
+
+		for (TSoftObjectPtr<UObject> SoftObjectPtr : Assets)
+		{
+			SoftObjectPathsToLoad.Add(SoftObjectPtr.ToSoftObjectPath());
+		}
+
+		// We always spawn a new load even if this node already queued one, the outside node handles this case
+		FLoadAssetsAction* NewAction = new FLoadAssetsAction(SoftObjectPathsToLoad, OnLoaded, LatentInfo);
+		LatentManager.AddNewAction(LatentInfo.CallbackTarget, LatentInfo.UUID, NewAction);
+	}
+}
+
 void UKismetSystemLibrary::RegisterForRemoteNotifications()
 {
 	FPlatformMisc::RegisterForRemoteNotifications();
