UE Blueprints: Fix duplicate loc ID for recently added "LoadAssets" node

diff --git a/Engine/Source/Editor/BlueprintGraph/Private/K2Node_LoadAsset.cpp b/Engine/Source/Editor/BlueprintGraph/Private/K2Node_LoadAsset.cpp
--- a/Engine/Source/Editor/BlueprintGraph/Private/K2Node_LoadAsset.cpp
+++ b/Engine/Source/Editor/BlueprintGraph/Private/K2Node_LoadAsset.cpp
@@ -357,7 +357,7 @@ FText UK2Node_LoadAssets::GetTooltipText() const
 
 FText UK2Node_LoadAssets::GetNodeTitle(ENodeTitleType::Type TitleType) const
 {
-	return FText(LOCTEXT("UK2Node_LoadAssetClassGetNodeTitle", "Async Load Assets"));
+	return FText(LOCTEXT("UK2Node_LoadAssetsGetNodeTitle", "Async Load Assets"));
 }
 
 FName UK2Node_LoadAssets::NativeFunctionName() const
