Fixed use after free error in BP compiler.

Manifested as generating incorrect bytecode with new compact set implementation.

#tests PIE

diff --git a/Engine/Source/Editor/BlueprintGraph/Private/CallFunctionHandler.cpp b/Engine/Source/Editor/BlueprintGraph/Private/CallFunctionHandler.cpp
--- a/Engine/Source/Editor/BlueprintGraph/Private/CallFunctionHandler.cpp
+++ b/Engine/Source/Editor/BlueprintGraph/Private/CallFunctionHandler.cpp
@@ -351,40 +351,36 @@ void FKCHandler_CallFunction::CreateFunctionCallStatement(FKismetFunctionContext
 									RHSTerm = *InterfaceTerm;
 								}
 
+								if (const CastingUtils::FImplicitCastParams* CastParams = Context.ImplicitCastMap.Find(PinMatch))
 								{
-									const CastingUtils::FImplicitCastParams* CastParams =
-										Context.ImplicitCastMap.Find(PinMatch);
+									CastingUtils::InsertImplicitCastStatement(Context, *CastParams, RHSTerm);
 
-									if (CastParams)
-									{
-										FBPTerminal* ImplicitCastTerm =
-											CastingUtils::InsertImplicitCastStatement(Context, PinMatch, RHSTerm);
-										check(ImplicitCastTerm);
-
-										bool bIsNonConstReference =
-											Property->HasAllPropertyFlags(CPF_OutParm | CPF_ReferenceParm) &&
-											!Property->HasAllPropertyFlags(CPF_ConstParm);
-
-										if (bIsNonConstReference)
-										{
-											CastingUtils::FImplicitCastParams InverseCastParams = *CastParams;
+									FBPTerminal* ImplicitCastTerm = CastParams->TargetTerminal;
+									check(ImplicitCastTerm);
 
-											if (CastParams->Conversion.Type == CastingUtils::FloatingPointCastType::FloatToDouble)
-											{
-												InverseCastParams.Conversion.Type = CastingUtils::FloatingPointCastType::DoubleToFloat;
-											}
-											else if (CastParams->Conversion.Type == CastingUtils::FloatingPointCastType::DoubleToFloat)
-											{
-												InverseCastParams.Conversion.Type = CastingUtils::FloatingPointCastType::FloatToDouble;
-											}
+									const bool bIsNonConstReference =
+										Property->HasAllPropertyFlags(CPF_OutParm | CPF_ReferenceParm) &&
+										!Property->HasAllPropertyFlags(CPF_ConstParm);
 
-											InverseCastParams.TargetTerminal = RHSTerm;
+									if (bIsNonConstReference)
+									{
+										CastingUtils::FImplicitCastParams& InverseCastParams = ModifiedCastInputs.Emplace_GetRef(ImplicitCastTerm, *CastParams).Value;
 
-											ModifiedCastInputs.Add(CastEntryT{ImplicitCastTerm, InverseCastParams});
+										const CastingUtils::FloatingPointCastType CastType = CastParams->Conversion.Type;
+										if (CastType == CastingUtils::FloatingPointCastType::FloatToDouble)
+										{
+											InverseCastParams.Conversion.Type = CastingUtils::FloatingPointCastType::DoubleToFloat;
 										}
-
-										RHSTerm = ImplicitCastTerm;
+										else if (CastType == CastingUtils::FloatingPointCastType::DoubleToFloat)
+										{
+											InverseCastParams.Conversion.Type = CastingUtils::FloatingPointCastType::FloatToDouble;
+										}
+										InverseCastParams.TargetTerminal = RHSTerm;
 									}
+
+									Context.ImplicitCastMap.Remove(PinMatch);
+
+									RHSTerm = ImplicitCastTerm;
 								}
 
 								int32 ParameterIndex = RHSTerms.Add(RHSTerm);
