Add a virtual function to UBlueprint, ShouldAutomaticallyRegisterInputOnConstruction, which can be overridden in widget blueprints.

This will allow us to check if we should set bAutomaticallyRegisterInputOnConstruction on PostCDO compiled for nodes which are in the base Engine module, where the UWidgetBlueprint type doesn't exist and we don't want to add a dependency.

In the future, this will also allow us to determine whether or not an actor / actor component would require input delegates to be bound as well, removing the need for any calls to "Enable Input" on owning actors and allowing us to handle it automatically.

diff --git a/Engine/Source/Editor/BlueprintGraph/Classes/K2Node_InputAxisKeyEvent.h b/Engine/Source/Editor/BlueprintGraph/Classes/K2Node_InputAxisKeyEvent.h
--- a/Engine/Source/Editor/BlueprintGraph/Classes/K2Node_InputAxisKeyEvent.h
+++ b/Engine/Source/Editor/BlueprintGraph/Classes/K2Node_InputAxisKeyEvent.h
@@ -65,6 +65,7 @@ class UK2Node_InputAxisKeyEvent : public UK2Node_Event
 	UE_API virtual void GetMenuActions(FBlueprintActionDatabaseRegistrar& ActionRegistrar) const override;
 	UE_API virtual FText GetMenuCategory() const override;
 	UE_API virtual FBlueprintNodeSignature GetSignature() const override;
+	UE_API virtual void ExpandNode(FKismetCompilerContext& CompilerContext, UEdGraph* SourceGraph) override;
 	//~ End UK2Node Interface
 
 	UE_API void Initialize(const FKey AxisKey);
diff --git a/Engine/Source/Editor/BlueprintGraph/Private/K2Node_InputAxisKeyEvent.cpp b/Engine/Source/Editor/BlueprintGraph/Private/K2Node_InputAxisKeyEvent.cpp
--- a/Engine/Source/Editor/BlueprintGraph/Private/K2Node_InputAxisKeyEvent.cpp
+++ b/Engine/Source/Editor/BlueprintGraph/Private/K2Node_InputAxisKeyEvent.cpp
@@ -4,6 +4,7 @@
 
 #include "BlueprintActionDatabaseRegistrar.h"
 #include "BlueprintNodeSpawner.h"
+#include "Blueprint/UserWidget.h"
 #include "Containers/Array.h"
 #include "Containers/Map.h"
 #include "Containers/UnrealString.h"
@@ -19,6 +20,7 @@
 #include "Internationalization/Internationalization.h"
 #include "Kismet2/BlueprintEditorUtils.h"
 #include "Kismet2/CompilerResultsLog.h"
+#include "KismetCompiler.h"
 #include "Misc/AssertionMacros.h"
 #include "Serialization/Archive.h"
 #include "Styling/AppStyle.h"
@@ -207,4 +209,21 @@ FBlueprintNodeSignature UK2Node_InputAxisKeyEvent::GetSignature() const
 	return NodeSignature;
 }
 
+void UK2Node_InputAxisKeyEvent::ExpandNode(FKismetCompilerContext& CompilerContext, UEdGraph* SourceGraph)
+{
+	Super::ExpandNode(CompilerContext, SourceGraph);	
+	
+	// Widget blueprints require the bAutomaticallyRegisterInputOnConstruction to be set to true in order to receive callbacks
+	if (GetBlueprint()->ShouldAutomaticallyRegisterInputOnConstruction())
+	{
+		CompilerContext.AddPostCDOCompiledStep([](const UObject::FPostCDOCompiledContext& Context, UObject* NewCDO)
+		{
+			if (UUserWidget* Widget = Cast<UUserWidget>(NewCDO))
+			{
+				Widget->bAutomaticallyRegisterInputOnConstruction = true;	
+			}
+		});	
+	}
+}
+
 #undef LOCTEXT_NAMESPACE
diff --git a/Engine/Source/Runtime/Engine/Classes/Engine/Blueprint.h b/Engine/Source/Runtime/Engine/Classes/Engine/Blueprint.h
--- a/Engine/Source/Runtime/Engine/Classes/Engine/Blueprint.h
+++ b/Engine/Source/Runtime/Engine/Classes/Engine/Blueprint.h
@@ -695,6 +695,8 @@ PRAGMA_ENABLE_DEPRECATION_WARNINGS
 
 	/** Whether or not this blueprint can be considered for a bytecode only compile */
 	virtual bool IsValidForBytecodeOnlyRecompile() const { return true; }
+	
+	virtual bool ShouldAutomaticallyRegisterInputOnConstruction() const { return false; }
 
 #if WITH_EDITORONLY_DATA
 	/** Delegate called when the debug object is set */
diff --git a/Engine/Source/Runtime/UMG/Public/Blueprint/UserWidgetBlueprint.h b/Engine/Source/Runtime/UMG/Public/Blueprint/UserWidgetBlueprint.h
--- a/Engine/Source/Runtime/UMG/Public/Blueprint/UserWidgetBlueprint.h
+++ b/Engine/Source/Runtime/UMG/Public/Blueprint/UserWidgetBlueprint.h
@@ -13,4 +13,6 @@ class UUserWidgetBlueprint : public UBlueprint
 public:
 	/** Does the editor support widget from an editor package. */
 	virtual bool AllowEditorWidget() const { return false; }
+
+	virtual bool ShouldAutomaticallyRegisterInputOnConstruction() const override { return true; }
 };
