UE-333131 - SetFields node won't accept non-const reference parameter nodes

The FProperties for a function that's overridden in Blueprint will inherit the flags from its parent's signature. However, the compiler also adds a CPF_BlueprintReadOnly the properties to prevent accidental writes to these properties in the function graph. This is flag is also added to input/output parameters (ie: non-const references).

Node validation in UK2Node_SetFieldsInStruct doesn't account for cases where we have a non-const reference, which is perfectly legal to use here. Instead of simply checking against CPF_BlueprintReadOnly, we need to have it also check the CPF_ReturnParm, CPF_OutParm, and CPF_ReferenceParm flags.

diff --git a/Engine/Source/Editor/BlueprintGraph/Private/K2Node_SetFieldsInStruct.cpp b/Engine/Source/Editor/BlueprintGraph/Private/K2Node_SetFieldsInStruct.cpp
--- a/Engine/Source/Editor/BlueprintGraph/Private/K2Node_SetFieldsInStruct.cpp
+++ b/Engine/Source/Editor/BlueprintGraph/Private/K2Node_SetFieldsInStruct.cpp
@@ -258,7 +258,18 @@ void UK2Node_SetFieldsInStruct::ValidateNodeDuringCompilation(FCompilerResultsLo
 			{
 				if (FProperty* BoundProperty = GetterNode->GetPropertyForVariable())
 				{
-					if (BoundProperty->HasAnyPropertyFlags(CPF_BlueprintReadOnly))
+					const bool bIsFunctionInput =
+						!BoundProperty->HasAnyPropertyFlags(CPF_ReturnParm) &&
+						(!BoundProperty->HasAnyPropertyFlags(CPF_OutParm) || BoundProperty->HasAnyPropertyFlags(CPF_ReferenceParm))
+					;
+
+					const bool bIsReadOnly =
+						BoundProperty->HasAnyPropertyFlags(CPF_BlueprintReadOnly)
+					;
+
+					// Function input parameters generally get tagged as readonly by the compiler.
+					// If they're actually input/ouput reference parameters, then we can make an exception here.
+					if (bIsReadOnly && !bIsFunctionInput)
 					{
 						// TODO, This should REALLY be an error, but too much code may have been written not following this standard.
 						FText ErrorMessage = LOCTEXT("SetStructFields_StructIsConst", "The @@ is a Read Only property and can not be modified directly.");
