Moved UEdGraphSchema asset reference filtering for drag and drop operations to their various implementations, which resulted in the following changes:
- Asset reference filtering now occurs when dropping on graphs, nodes and pins (previously it only occurred when dropping on graphs).
- Implementations now typically filter out assets they don't support first resulting in more relevant tooltips.  E.g. it is more useful to know a graph doesn't support a specific asset type than the asset cannot be referenced since you still wont be able to drop the asset if you move it.
- The UEdGraphSchema doesn't have knowledge of what the graph represents, so the high level assumption that dropping an asset on the graph results in a reference to that asset can be wrong.  E.g. in the USoundSubmixGraphSchema, dropping additional sound submixes opens them up for edit in the same graph.

Localized various hover messages in different UEdGraphSchema implementations.

Dropping an asset on the Plugin Reference Viewer now center's the graph on the asset's plugin.

Fixed the asset reference filter not working on Materials, Niagara Scripts and Texture Graphs since they operate on temporary copy of the asset.
Fixed the UAnimationStateMachineSchema graph drop tooltip indicating you can't drop UAnimationAssets on it even though you can.
Fixed the USoundCueGraphSchema graph drop tooltip indicating you can't drop UDialogueWaves on it even though you can.
Fixed the UTG_EdGraphSchema graph drop tooltip indicating you can't drop UMaterialFunctionInterfaces on it even though you can.

diff --git a/Engine/Plugins/Enterprise/DataprepEditor/Source/DataprepEditor/Private/DataprepGraph/DataprepGraphSchema.h b/Engine/Plugins/Enterprise/DataprepEditor/Source/DataprepEditor/Private/DataprepGraph/DataprepGraphSchema.h
--- a/Engine/Plugins/Enterprise/DataprepEditor/Source/DataprepEditor/Private/DataprepGraph/DataprepGraphSchema.h
+++ b/Engine/Plugins/Enterprise/DataprepEditor/Source/DataprepEditor/Private/DataprepGraph/DataprepGraphSchema.h
@@ -16,8 +16,6 @@ public:
 	// UEdGraphSchema interface
 	virtual void GetContextMenuActions(class UToolMenu* Menu, class UGraphNodeContextMenuContext* Context) const override {}
 	virtual FName GetParentContextMenuName() const override { return NAME_None; }
-	virtual void DroppedAssetsOnGraph(const TArray<FAssetData>& Assets, const FVector2f& GraphPosition, UEdGraph* Graph) const override {}
-	virtual void GetAssetsGraphHoverMessage(const TArray<FAssetData>& Assets, const UEdGraph* HoverGraph, FString& OutTooltipText, bool& OutOkIcon) const override {}
 	//virtual class FConnectionDrawingPolicy* CreateConnectionDrawingPolicy(int32 InBackLayerID, int32 InFrontLayerID, float InZoomFactor, const FSlateRect& InClippingRect, class FSlateWindowElementList& InDrawElements, class UEdGraph* InGraphObj) const override;
 	// End of UEdGraphSchema interface
 };
diff --git a/Engine/Plugins/Experimental/Dataflow/Source/DataflowEditor/Private/Dataflow/DataflowAssetEditUtils.cpp b/Engine/Plugins/Experimental/Dataflow/Source/DataflowEditor/Private/Dataflow/DataflowAssetEditUtils.cpp
--- a/Engine/Plugins/Experimental/Dataflow/Source/DataflowEditor/Private/Dataflow/DataflowAssetEditUtils.cpp
+++ b/Engine/Plugins/Experimental/Dataflow/Source/DataflowEditor/Private/Dataflow/DataflowAssetEditUtils.cpp
@@ -13,6 +13,7 @@
 #include "Dataflow/DataflowSubGraphNodes.h"
 #include "Dataflow/DataflowVariableNodes.h"
 #include "Editor.h"
+#include "Editor/EditorEngine.h"
 #include "PropertyBagDetails.h"
 #include "ScopedTransaction.h"
 #include "Misc/StringOutputDevice.h"
@@ -376,6 +377,21 @@ namespace UE::Dataflow
 		return ::IsUniqueObjectName(SubObjectName, DataflowAsset);
 	}
 
+	TSharedPtr<IAssetReferenceFilter> FEditAssetUtils::MakeAssetReferenceFilter(const UEdGraph* Graph)
+	{
+		if (const UDataflow* Dataflow = UDataflow::GetDataflowAssetFromEdGraph(Graph))
+		{
+			if (GEditor)
+			{
+				FAssetReferenceFilterContext AssetReferenceFilterContext;
+				AssetReferenceFilterContext.AddReferencingAsset(Dataflow);
+				return GEditor->MakeAssetReferenceFilter(AssetReferenceFilterContext);
+			}
+		}
+
+		return {};
+	}
+
 	//////////////////////////////////////////////////////////////////////////////////////////////////////////
 	// 
 	// NODE API
diff --git a/Engine/Plugins/Experimental/Dataflow/Source/DataflowEditor/Private/Dataflow/DataflowAssetEditUtils.h b/Engine/Plugins/Experimental/Dataflow/Source/DataflowEditor/Private/Dataflow/DataflowAssetEditUtils.h
--- a/Engine/Plugins/Experimental/Dataflow/Source/DataflowEditor/Private/Dataflow/DataflowAssetEditUtils.h
+++ b/Engine/Plugins/Experimental/Dataflow/Source/DataflowEditor/Private/Dataflow/DataflowAssetEditUtils.h
@@ -12,6 +12,7 @@ struct FEdGraphPinType;
 struct FDataflowNode;
 struct FInstancedPropertyBag;
 struct FPropertyBagPropertyDesc;
+class IAssetReferenceFilter;
 class UDataflow;
 class UDataflowEdNode;
 class UEdGraph;
@@ -34,6 +35,8 @@ namespace UE::Dataflow
 		*/
 		static bool IsUniqueDataflowSubObjectName(UDataflow* DataflowAsset, FName Name);
 
+		static TSharedPtr<IAssetReferenceFilter> MakeAssetReferenceFilter(const UEdGraph* Graph);
+
 		/////////////////////////////////////////////////////////////////////////////////////////////////////////////////////////
 		//
 		// Nodes API
diff --git a/Engine/Plugins/Experimental/Dataflow/Source/DataflowEditor/Private/Dataflow/DataflowSchema.cpp b/Engine/Plugins/Experimental/Dataflow/Source/DataflowEditor/Private/Dataflow/DataflowSchema.cpp
--- a/Engine/Plugins/Experimental/Dataflow/Source/DataflowEditor/Private/Dataflow/DataflowSchema.cpp
+++ b/Engine/Plugins/Experimental/Dataflow/Source/DataflowEditor/Private/Dataflow/DataflowSchema.cpp
@@ -10,6 +10,7 @@
 #include "Dataflow/DataflowGraphSchemaAction.h"
 #include "Dataflow/DataflowNodeFactory.h"
 #include "Dataflow/DataflowAssetEditUtils.h"
+#include "Editor/AssetReferenceFilter.h"
 #include "EdGraph/EdGraphNode.h"
 #include "EdGraph/EdGraph.h"
 #include "Framework/Commands/GenericCommands.h"
@@ -559,12 +560,15 @@ UE::Dataflow::FPin::EDirection UDataflowSchema::GetDirectionFromPinDirection(EEd
 	}
 }
 
-void UDataflowSchema::GetAssetsGraphHoverMessage(const TArray<FAssetData>&Assets, const UEdGraph * HoverGraph, FString & OutTooltipText, bool& OutOkIcon) const
+void UDataflowSchema::GetAssetsGraphHoverMessage(const TArray<FAssetData>&Assets, const UEdGraph * HoverGraph, FString& OutTooltipText, bool& OutOkIcon) const
 {
 	using FDataflowNodeFactory = UE::Dataflow::FNodeFactory;
 
+	const TSharedPtr<IAssetReferenceFilter> AssetReferenceFilter = UE::Dataflow::FEditAssetUtils::MakeAssetReferenceFilter(HoverGraph);
+
+	OutTooltipText.Reset();
 	OutOkIcon = false;
-	OutTooltipText = FString(TEXT("No supported node found"));
+	
 	for (const FAssetData& AssetData : Assets)
 	{
 		if (UObject* AssetObject = AssetData.GetAsset())
@@ -572,17 +576,37 @@ void UDataflowSchema::GetAssetsGraphHoverMessage(const TArray<FAssetData>&Assets
 			const FName GetterNodeType = FDataflowNodeFactory::GetInstance()->GetGetterNodeFromAssetClass(*AssetObject->GetClass());
 			if (!GetterNodeType.IsNone())
 			{
+				if (AssetReferenceFilter)
+				{
+					FText FailureReason;
+					if (!AssetReferenceFilter->PassesFilter(AssetData, &FailureReason))
+					{
+						if (OutTooltipText.IsEmpty())
+						{
+							OutTooltipText = FailureReason.ToString();
+						}
+						continue;
+					}
+				}
+
+				OutTooltipText = LOCTEXT("DataflowSchema_DropAssetOnGraphSupported", "Place as a getter node here.").ToString();
 				OutOkIcon = true;
-				OutTooltipText = FString(TEXT("Place as a getter node here."));
-				break; // at least one is require to be a positive drop
+				return; // at least one is require to be a positive drop
 			}
 		}
 	}
+
+	if (OutTooltipText.IsEmpty())
+	{
+		OutTooltipText = LOCTEXT("DataflowSchema_NoDropAssetOnGraphSupported", "No supported node found").ToString();
+	}
 }
+
 void UDataflowSchema::DroppedAssetsOnGraph(const TArray<FAssetData>& Assets, const FVector2f& GraphPosition, UEdGraph* Graph) const
 {
 	using FDataflowNodeFactory = UE::Dataflow::FNodeFactory;
 
+	const TSharedPtr<IAssetReferenceFilter> AssetReferenceFilter = UE::Dataflow::FEditAssetUtils::MakeAssetReferenceFilter(Graph);
 	const FVector2f NodeOffset{ 0.0f,100.0f };
 	FVector2f NodePosition = GraphPosition;
 
@@ -593,6 +617,11 @@ void UDataflowSchema::DroppedAssetsOnGraph(const TArray<FAssetData>& Assets, con
 			const FName GetterNodeType = FDataflowNodeFactory::GetInstance()->GetGetterNodeFromAssetClass(*AssetObject->GetClass());
 			if (!GetterNodeType.IsNone())
 			{
+				if (AssetReferenceFilter.IsValid() && !AssetReferenceFilter->PassesFilter(AssetData))
+				{
+					continue;
+				}
+
 				if (UDataflowEdNode* DataflowEdNode = UE::Dataflow::FEditAssetUtils::AddNewNode(Graph, FDeprecateSlateVector2D(NodePosition), AssetObject->GetFName(), GetterNodeType, nullptr))
 				{
 					NodePosition += NodeOffset;
@@ -608,10 +637,12 @@ void UDataflowSchema::DroppedAssetsOnGraph(const TArray<FAssetData>& Assets, con
 
 void UDataflowSchema::GetAssetsNodeHoverMessage(const TArray<struct FAssetData>& Assets, const UEdGraphNode* HoverNode, FString& OutTooltipText, bool& OutOkIcon) const
 {
+	OutTooltipText.Reset();
 	OutOkIcon = false;
-	OutTooltipText = FString(TEXT("Asset type unsupported by this node"));
+
 	if (TSharedPtr<const FDataflowNode> DataflowNode = UDataflowEdNode::GetDataflowNodeFromEdNode(HoverNode))
 	{
+		const TSharedPtr<IAssetReferenceFilter> AssetReferenceFilter = UE::Dataflow::FEditAssetUtils::MakeAssetReferenceFilter(HoverNode->GetGraph());
 		bool bNeedInvalidation = false;
 		for (const FAssetData& AssetData : Assets)
 		{
@@ -619,19 +650,38 @@ void UDataflowSchema::GetAssetsNodeHoverMessage(const TArray<struct FAssetData>&
 			{
 				if (DataflowNode->SupportsAssetProperty(AssetObject))
 				{
+					if (AssetReferenceFilter)
+					{
+						FText FailureReason;
+						if (!AssetReferenceFilter->PassesFilter(AssetData, &FailureReason))
+						{
+							if (OutTooltipText.IsEmpty())
+							{
+								OutTooltipText = FailureReason.ToString();
+							}
+							continue;
+						}
+					}
+
+					OutTooltipText = LOCTEXT("DataflowSchema_DropAssetOnNodeSupported", "Set asset property on this node.").ToString();
 					OutOkIcon = true;
-					OutTooltipText = FString(TEXT("Set asset property on this node."));
-					break; // at least one is require to be a positive drop
+					return; // at least one is require to be a positive drop
 				}
 			}
 		}
 	}
+
+	if (OutTooltipText.IsEmpty())
+	{
+		OutTooltipText = LOCTEXT("DataflowSchema_NoDropAssetOnNodeSupported", "Asset type unsupported by this node").ToString();
+	}
 }
 
 void UDataflowSchema::DroppedAssetsOnNode(const TArray<struct FAssetData>& Assets, const FVector2f& GraphPosition, UEdGraphNode* Node) const
 {
 	if (TSharedPtr<FDataflowNode> DataflowNode = UDataflowEdNode::GetDataflowNodeFromEdNode(Node))
 	{
+		const TSharedPtr<IAssetReferenceFilter> AssetReferenceFilter = UE::Dataflow::FEditAssetUtils::MakeAssetReferenceFilter(Node->GetGraph());
 		bool bNeedInvalidation = false;
 		for (const FAssetData& AssetData : Assets)
 		{
@@ -639,8 +689,14 @@ void UDataflowSchema::DroppedAssetsOnNode(const TArray<struct FAssetData>& Asset
 			{
 				if (DataflowNode->SupportsAssetProperty(AssetObject))
 				{
+					if (AssetReferenceFilter.IsValid() && !AssetReferenceFilter->PassesFilter(AssetData))
+					{
+						continue;
+					}
+
 					DataflowNode->SetAssetProperty(AssetObject);
 					bNeedInvalidation = true;
+					break;
 				}
 			}
 		}
diff --git a/Engine/Plugins/Experimental/PluginReferenceViewer/Source/PluginReferenceViewer/Private/PluginReferenceViewerSchema.cpp b/Engine/Plugins/Experimental/PluginReferenceViewer/Source/PluginReferenceViewer/Private/PluginReferenceViewerSchema.cpp
--- a/Engine/Plugins/Experimental/PluginReferenceViewer/Source/PluginReferenceViewer/Private/PluginReferenceViewerSchema.cpp
+++ b/Engine/Plugins/Experimental/PluginReferenceViewer/Source/PluginReferenceViewer/Private/PluginReferenceViewerSchema.cpp
@@ -3,10 +3,14 @@
 #include "PluginReferenceViewerSchema.h"
 
 #include "AssetManagerEditorModule.h"
+#include "AssetRegistry/AssetData.h"
 #include "ConnectionDrawingPolicy.h"
 #include "EdGraphNode_PluginReference.h"
+#include "Interfaces/IPluginManager.h"
+#include "Misc/PathViews.h"
 #include "PluginReferencePinCategory.h"
 #include "PluginReferenceViewerCommands.h"
+#include "PluginReferenceViewerModule.h"
 #include "SPluginReferenceViewer.h"
 #include "ToolMenu.h"
 #include "ToolMenuEntry.h"
@@ -177,9 +181,37 @@ FConnectionDrawingPolicy* UPluginReferenceViewerSchema::CreateConnectionDrawingP
 
 void UPluginReferenceViewerSchema::DroppedAssetsOnGraph(const TArray<FAssetData>& Assets, const FVector2f& GraphPosition, UEdGraph* Graph) const
 {
+	if (TSharedPtr<IPlugin> Plugin = GetDroppedPlugin(Assets))
+	{
+		FPluginReferenceViewerModule::Get().OpenPluginReferenceViewerUI(MoveTemp(Plugin).ToSharedRef());
+	}
 }
 
 void UPluginReferenceViewerSchema::GetAssetsGraphHoverMessage(const TArray<FAssetData>& Assets, const UEdGraph* HoverGraph, FString& OutTooltipText, bool& OutOkIcon) const
 {
-	OutOkIcon = true;
+	if (TSharedPtr<IPlugin> Plugin = GetDroppedPlugin(Assets))
+	{
+		OutTooltipText = FText::FormatOrdered(NSLOCTEXT("PluginReferenceViewerSchema", "RecenterGraphOnPlugin", "Re-center graph on plugin {0}."), FText::FromString(Plugin->GetFriendlyName())).ToString();
+		OutOkIcon = true;
+	}
+	else
+	{
+		OutTooltipText = NSLOCTEXT("PluginReferenceViewerSchema", "NoPlugin", "Asset does not belong to a plugin.").ToString();
+		OutOkIcon = false;
+	}
+}
+
+TSharedPtr<IPlugin> UPluginReferenceViewerSchema::GetDroppedPlugin(const TArray<FAssetData>& Assets)
+{
+	for (const FAssetData& AssetData : Assets)
+	{
+		const FNameBuilder PackageName(AssetData.PackageName);
+		FStringView PluginName = FPathViews::GetMountPointNameFromPath(PackageName);
+		if (TSharedPtr<IPlugin> Plugin = IPluginManager::Get().FindPlugin(PluginName))
+		{
+			return Plugin;
+		}
+	}
+
+	return {};
 }
\ No newline at end of file
diff --git a/Engine/Plugins/Experimental/PluginReferenceViewer/Source/PluginReferenceViewer/Private/PluginReferenceViewerSchema.h b/Engine/Plugins/Experimental/PluginReferenceViewer/Source/PluginReferenceViewer/Private/PluginReferenceViewerSchema.h
--- a/Engine/Plugins/Experimental/PluginReferenceViewer/Source/PluginReferenceViewer/Private/PluginReferenceViewerSchema.h
+++ b/Engine/Plugins/Experimental/PluginReferenceViewer/Source/PluginReferenceViewer/Private/PluginReferenceViewerSchema.h
@@ -7,9 +7,10 @@
 #include "EdGraph/EdGraphSchema.h"
 #include "PluginReferenceViewerSchema.generated.h"
 
-class UToolMenu;
 class FSlateRect;
+class IPlugin;
 class UEdGraph;
+class UToolMenu;
 
 UCLASS(MinimalAPI)
 class UPluginReferenceViewerSchema : public UEdGraphSchema
@@ -29,4 +30,7 @@ public:
 	virtual void DroppedAssetsOnGraph(const TArray<FAssetData>& Assets, const FVector2f& GraphPosition, UEdGraph* Graph) const override;
 	virtual void GetAssetsGraphHoverMessage(const TArray<FAssetData>& Assets, const UEdGraph* HoverGraph, FString& OutTooltipText, bool& OutOkIcon) const override;
 	// End of UEdGraphSchema interface
+
+private:
+	static TSharedPtr<IPlugin> GetDroppedPlugin(const TArray<FAssetData>& Assets);
 };
\ No newline at end of file
diff --git a/Engine/Plugins/Experimental/UAF/UAF/Source/UAFUncookedOnly/Internal/AnimNextEdGraphSchema.h b/Engine/Plugins/Experimental/UAF/UAF/Source/UAFUncookedOnly/Internal/AnimNextEdGraphSchema.h
--- a/Engine/Plugins/Experimental/UAF/UAF/Source/UAFUncookedOnly/Internal/AnimNextEdGraphSchema.h
+++ b/Engine/Plugins/Experimental/UAF/UAF/Source/UAFUncookedOnly/Internal/AnimNextEdGraphSchema.h
@@ -7,6 +7,8 @@
 #include "AnimNextEdGraphNode.h"
 #include "AnimNextEdGraphSchema.generated.h"
 
+class IAssetReferenceFilter;
+
 #define UE_API UAFUNCOOKEDONLY_API
 
 UCLASS(MinimalAPI)
@@ -19,6 +21,9 @@ class UAnimNextEdGraphSchema : public URigVMEdGraphSchema
 
 	// UEdGraphSchema interface
 	UE_API virtual void GetGraphDisplayInformation(const UEdGraph& Graph, /*out*/ FGraphDisplayInfo& DisplayInfo) const override;
+
+protected:
+	UE_API static TSharedPtr<IAssetReferenceFilter> MakeAssetReferenceFilter(const UEdGraph* Graph);
 };
 
 #undef UE_API
\ No newline at end of file
diff --git a/Engine/Plugins/Experimental/UAF/UAF/Source/UAFUncookedOnly/Private/AnimNextEdGraphSchema.cpp b/Engine/Plugins/Experimental/UAF/UAF/Source/UAFUncookedOnly/Private/AnimNextEdGraphSchema.cpp
--- a/Engine/Plugins/Experimental/UAF/UAF/Source/UAFUncookedOnly/Private/AnimNextEdGraphSchema.cpp
+++ b/Engine/Plugins/Experimental/UAF/UAF/Source/UAFUncookedOnly/Private/AnimNextEdGraphSchema.cpp
@@ -3,6 +3,8 @@
 #include "AnimNextEdGraphSchema.h"
 
 #include "AnimNextRigVMAsset.h"
+#include "Editor.h"
+#include "Editor/EditorEngine.h"
 #include "Entries/AnimNextRigVMAssetEntry.h"
 
 #include UE_INLINE_GENERATED_CPP_BY_NAME(AnimNextEdGraphSchema)
@@ -20,4 +22,22 @@ void UAnimNextEdGraphSchema::GetGraphDisplayInformation(const UEdGraph& Graph, /
 	}
 }
 
+TSharedPtr<IAssetReferenceFilter> UAnimNextEdGraphSchema::MakeAssetReferenceFilter(const UEdGraph* Graph)
+{
+	if (Graph)
+	{
+		if (UAnimNextRigVMAsset* Asset = Cast<UAnimNextRigVMAsset>(Graph->GetTypedOuter<UAnimNextRigVMAsset>()))
+		{
+			if (GEditor)
+			{
+				FAssetReferenceFilterContext AssetReferenceFilterContext;
+				AssetReferenceFilterContext.AddReferencingAsset(Asset);
+				return GEditor->MakeAssetReferenceFilter(AssetReferenceFilterContext);
+			}
+		}
+	}
+
+	return {};
+}
+
 #undef LOCTEXT_NAMESPACE
diff --git a/Engine/Plugins/Experimental/UAF/UAFAnimGraph/Source/UAFAnimGraphUncookedOnly/Private/Graph/UAFAnimGraphEdGraphSchema.cpp b/Engine/Plugins/Experimental/UAF/UAFAnimGraph/Source/UAFAnimGraphUncookedOnly/Private/Graph/UAFAnimGraphEdGraphSchema.cpp
--- a/Engine/Plugins/Experimental/UAF/UAFAnimGraph/Source/UAFAnimGraphUncookedOnly/Private/Graph/UAFAnimGraphEdGraphSchema.cpp
+++ b/Engine/Plugins/Experimental/UAF/UAFAnimGraph/Source/UAFAnimGraphUncookedOnly/Private/Graph/UAFAnimGraphEdGraphSchema.cpp
@@ -4,6 +4,7 @@
 
 #include "AnimNextController.h"
 #include "EdGraph/RigVMEdGraph.h"
+#include "Editor/AssetReferenceFilter.h"
 #include "Framework/Application/SlateApplication.h"
 #include "Framework/MultiBox/MultiBoxBuilder.h"
 #include "Templates/GraphNodeTemplateRegistry.h"
@@ -24,6 +25,8 @@ void UUAFAnimGraphEdGraphSchema::DroppedAssetsOnGraph(const TArray<FAssetData>&
 		return;
 	}
 
+	const TSharedPtr<IAssetReferenceFilter> AssetReferenceFilter = MakeAssetReferenceFilter(Graph);
+
 	auto HandleDrop = [Controller](UUAFGraphNodeTemplate* InTemplate, UObject* InAsset, FVector2D InLocation)
 	{
 		Controller->OpenUndoBracket(TEXT("Drop assets on graph"));
@@ -50,6 +53,11 @@ void UUAFAnimGraphEdGraphSchema::DroppedAssetsOnGraph(const TArray<FAssetData>&
 	// See if any of the known templates know how to handle the dropped assets
 	for (const FAssetData& AssetData : Assets)
 	{
+		if (AssetReferenceFilter.IsValid() && !AssetReferenceFilter->PassesFilter(AssetData))
+		{
+			continue;
+		}
+
 		TConstArrayView<FGraphNodeTemplateInfo> Handlers = FGraphNodeTemplateRegistry::GetDragDropHandlersForAsset(AssetData);
 		if (Handlers.Num() > 1)
 		{
@@ -76,16 +84,22 @@ void UUAFAnimGraphEdGraphSchema::DroppedAssetsOnGraph(const TArray<FAssetData>&
 				FEditorScriptExecutionGuard AllowScripts;
 				FText MenuDesc = Template->GetMenuDescription();
 				FText Tooltip = Template->GetTooltipText();
-				MenuBuilder.AddMenuEntry(MenuDesc, Tooltip, FSlateIcon(), FExecuteAction::CreateLambda([Template, HandleDrop, Assets, GraphPosition]()
+				MenuBuilder.AddMenuEntry(MenuDesc, Tooltip, FSlateIcon(), FExecuteAction::CreateLambda([Template, AssetReferenceFilter, HandleDrop, Assets, GraphPosition]()
 				{
 					FVector2D Location = FVector2D(GraphPosition);
 					for (const FAssetData& AssetData : Assets)
 					{
+						if (AssetReferenceFilter.IsValid() && !AssetReferenceFilter->PassesFilter(AssetData))
+						{
+							continue;
+						}
+
 						UObject* Asset = AssetData.GetAsset();
 						if (Asset == nullptr)
 						{
 							continue;
 						}
+
 						HandleDrop(Template, Asset, Location);
 
 						// Offset multiple spawned nodes
@@ -112,6 +126,11 @@ void UUAFAnimGraphEdGraphSchema::DroppedAssetsOnGraph(const TArray<FAssetData>&
 		FVector2D Location = FVector2D(GraphPosition);
 		for (const FAssetData& AssetData : Assets)
 		{
+			if (AssetReferenceFilter.IsValid() && !AssetReferenceFilter->PassesFilter(AssetData))
+			{
+				continue;
+			}
+
 			UObject* Asset = AssetData.GetAsset();
 			if (Asset == nullptr)
 			{
@@ -139,18 +158,34 @@ void UUAFAnimGraphEdGraphSchema::DroppedAssetsOnGraph(const TArray<FAssetData>&
 void UUAFAnimGraphEdGraphSchema::GetAssetsGraphHoverMessage(const TArray<FAssetData>& Assets, const UEdGraph* HoverGraph, FString& OutTooltipText, bool& OutOkIcon) const
 {
 	using namespace UE::UAF;
+
+	OutTooltipText.Reset();
+	OutOkIcon = false;
+
+	const TSharedPtr<IAssetReferenceFilter> AssetReferenceFilter = MakeAssetReferenceFilter(HoverGraph);
 	for (const FAssetData& AssetData : Assets)
 	{
 		TConstArrayView<FGraphNodeTemplateInfo> Handlers = FGraphNodeTemplateRegistry::GetDragDropHandlersForAsset(AssetData);
 		if (Handlers.Num() > 0)
 		{
+			if (AssetReferenceFilter)
+			{
+				FText FailureReason;
+				if (!AssetReferenceFilter->PassesFilter(AssetData, &FailureReason))
+				{
+					if (OutTooltipText.IsEmpty())
+					{
+						OutTooltipText = FailureReason.ToString();
+					}
+					continue;
+				}
+			}
+
 			OutOkIcon = true;
-			OutTooltipText = Assets.Num() > 0 ? LOCTEXT("DragDropAssetMessage", "Create node for asset").ToString() : LOCTEXT("DragDropAssetsMessage", "Create nodes for assets").ToString();
+			OutTooltipText = (Assets.Num() > 0 ? LOCTEXT("DragDropAssetMessage", "Create node for asset") : LOCTEXT("DragDropAssetsMessage", "Create nodes for assets")).ToString();
 			return;
 		}
 	}
-
-	Super::GetAssetsGraphHoverMessage(Assets, HoverGraph, OutTooltipText, OutOkIcon);
 }
 
 FLinearColor UUAFAnimGraphEdGraphSchema::GetPinTypeColor(const FEdGraphPinType& PinType) const
diff --git a/Engine/Plugins/FX/Niagara/Source/NiagaraEditor/Private/EdGraphSchema_Niagara.cpp b/Engine/Plugins/FX/Niagara/Source/NiagaraEditor/Private/EdGraphSchema_Niagara.cpp
--- a/Engine/Plugins/FX/Niagara/Source/NiagaraEditor/Private/EdGraphSchema_Niagara.cpp
+++ b/Engine/Plugins/FX/Niagara/Source/NiagaraEditor/Private/EdGraphSchema_Niagara.cpp
@@ -3,6 +3,7 @@
 #include "EdGraphSchema_Niagara.h"
 
 #include "AssetRegistry/AssetRegistryModule.h"
+#include "Editor/AssetReferenceFilter.h"
 #include "EdGraphNode_Comment.h"
 #include "GraphEditorSettings.h"
 #include "INiagaraEditorTypeUtilities.h"
@@ -40,6 +41,8 @@
 #include "ObjectEditorUtils.h"
 #include "ScopedTransaction.h"
 #include "ToolMenus.h"
+#include "ViewModels/NiagaraScriptViewModel.h"
+#include "ViewModels/TNiagaraViewModelManager.h"
 #include "Settings/EditorStyleSettings.h"
 #include "Framework/Commands/UIAction.h"
 #include "HAL/PlatformApplicationMisc.h"
@@ -1352,14 +1355,16 @@ void UEdGraphSchema_Niagara::DroppedAssetsOnGraph(const TArray<FAssetData>& Asse
 	uint32 Offset = 0;
 	TArray<UEnum*> Enums;
 
+	const TSharedPtr<IAssetReferenceFilter> AssetReferenceFilter = MakeAssetReferenceFilter(Graph);
 	for (const FAssetData& Data : Assets)
 	{
 		UObject* Asset = Data.GetAsset();
-		UEnum* Enum = Cast<UEnum>(Asset);
-
-		if (Enum)
+		if (UEnum* Enum = Cast<UEnum>(Asset))
 		{
-			Enums.Add(Enum);
+			if (!AssetReferenceFilter.IsValid() || AssetReferenceFilter->PassesFilter(Data))
+			{
+				Enums.Add(Enum);
+			}
 		}
 	}
 
@@ -1384,16 +1389,31 @@ void UEdGraphSchema_Niagara::DroppedAssetsOnGraph(const TArray<FAssetData>& Asse
 	}
 }
 
-void UEdGraphSchema_Niagara::GetAssetsGraphHoverMessage(const TArray<FAssetData>&Assets, const UEdGraph * HoverGraph, FString & OutTooltipText, bool& OutOkIcon) const
+void UEdGraphSchema_Niagara::GetAssetsGraphHoverMessage(const TArray<FAssetData>&Assets, const UEdGraph* HoverGraph, FString& OutTooltipText, bool& OutOkIcon) const
 {
+	OutTooltipText.Reset();
 	OutOkIcon = false;
 
+	const TSharedPtr<IAssetReferenceFilter> AssetReferenceFilter = MakeAssetReferenceFilter(HoverGraph);
 	for (const FAssetData& AssetData : Assets)
 	{
 		UEnum* Enum = Cast<UEnum>(AssetData.GetAsset());
 		if (Enum)
 		{
-			OutTooltipText = TEXT("Create a static switch using the selected enum");
+			if (AssetReferenceFilter)
+			{
+				FText FailureReason;
+				if (!AssetReferenceFilter->PassesFilter(AssetData, &FailureReason))
+				{
+					if (OutTooltipText.IsEmpty())
+					{
+						OutTooltipText = FailureReason.ToString();
+					}
+					continue;
+				}
+			}
+
+			OutTooltipText = LOCTEXT("CreateStaticMeshSwitch", "Create a static switch using the selected enum").ToString();
 			OutOkIcon = true;
 			break;
 		}
@@ -2479,6 +2499,28 @@ void UEdGraphSchema_Niagara::GetContextMenuActions(UToolMenu* Menu, UGraphNodeCo
 	Super::GetContextMenuActions(Menu, Context);
 }
 
+TSharedPtr<IAssetReferenceFilter> UEdGraphSchema_Niagara::MakeAssetReferenceFilter(const UEdGraph* Graph)
+{
+	if (Graph)
+	{
+		if (UNiagaraScript* Script = Graph->GetTypedOuter<UNiagaraScript>())
+		{
+			TArray<TSharedPtr<FNiagaraScriptViewModel>> ViewModels;
+			TNiagaraViewModelManager<UNiagaraScript, FNiagaraScriptViewModel>::GetAllViewModelsForObject(Script, ViewModels);
+			if (ViewModels.Num() == 1)
+			{
+				return ViewModels[0]->MakeAssetReferenceFilter();
+			}
+			else
+			{
+				ensureMsgf(false, TEXT("Missing or duplicate script view models detected.  Can not make IAssetReferenceFilter."));
+			}
+		}
+	}
+
+	return {};
+}
+
 FNiagaraConnectionDrawingPolicy::FNiagaraConnectionDrawingPolicy(int32 InBackLayerID, int32 InFrontLayerID, float InZoomFactor, const FSlateRect& InClippingRect, FSlateWindowElementList& InDrawElements, UEdGraph* InGraph)
 	: FConnectionDrawingPolicy(InBackLayerID, InFrontLayerID, InZoomFactor, InClippingRect, InDrawElements)
 	, GraphObj(CastChecked<UNiagaraGraph>(InGraph))
diff --git a/Engine/Plugins/FX/Niagara/Source/NiagaraEditor/Private/EdGraphSchema_NiagaraSystemOverview.cpp b/Engine/Plugins/FX/Niagara/Source/NiagaraEditor/Private/EdGraphSchema_NiagaraSystemOverview.cpp
--- a/Engine/Plugins/FX/Niagara/Source/NiagaraEditor/Private/EdGraphSchema_NiagaraSystemOverview.cpp
+++ b/Engine/Plugins/FX/Niagara/Source/NiagaraEditor/Private/EdGraphSchema_NiagaraSystemOverview.cpp
@@ -41,6 +41,7 @@ void UEdGraphSchema_NiagaraSystemOverview::ForceVisualizationCacheClear() const
 
 void UEdGraphSchema_NiagaraSystemOverview::DroppedAssetsOnGraph(const TArray<FAssetData>& Assets, const FVector2f& GraphPosition, UEdGraph* Graph) const
 {
+	const TSharedPtr<IAssetReferenceFilter> AssetReferenceFilter = MakeAssetReferenceFilter(Graph);
 	TArray<UNiagaraEmitter*> Emitters;
 
 	for (const FAssetData& Data : Assets)
@@ -48,7 +49,10 @@ void UEdGraphSchema_NiagaraSystemOverview::DroppedAssetsOnGraph(const TArray<FAs
 		UObject* Asset = Data.GetAsset();
 		if (UNiagaraEmitter* Emitter = Cast<UNiagaraEmitter>(Asset))
 		{
-			Emitters.Add(Emitter);
+			if (!AssetReferenceFilter.IsValid() || AssetReferenceFilter->PassesFilter(Asset))
+			{
+				Emitters.Add(Emitter);
+			}
 		}
 	}
 
@@ -84,6 +88,7 @@ void UEdGraphSchema_NiagaraSystemOverview::DroppedAssetsOnGraph(const TArray<FAs
 
 void UEdGraphSchema_NiagaraSystemOverview::GetAssetsGraphHoverMessage(const TArray<FAssetData>& Assets, const UEdGraph* HoverGraph, FString& OutTooltipText, bool& OutOkIcon) const
 {
+	OutTooltipText.Reset();
 	OutOkIcon = false;
 
 	UNiagaraSystem* System = HoverGraph->GetTypedOuter<UNiagaraSystem>();
@@ -98,25 +103,55 @@ void UEdGraphSchema_NiagaraSystemOverview::GetAssetsGraphHoverMessage(const TArr
 		{
 			if (UNiagaraEmitter* Emitter = Cast<UNiagaraEmitter>(AssetData.GetAsset()))
 			{
-				FText Tooltip = LOCTEXT("NiagaraEditorDropEmitterInEmitterTooltip", "Emitters can only be added to systems");
-				OutTooltipText = Tooltip.ToString();
+				OutTooltipText = LOCTEXT("NiagaraEditorDropEmitterInEmitterTooltip", "Emitters can only be added to systems").ToString();
 				break;
 			}
 		}
 		return;
 	}
 	
+	const TSharedPtr<IAssetReferenceFilter> AssetReferenceFilter = MakeAssetReferenceFilter(HoverGraph);
 	for (const FAssetData& AssetData : Assets)
 	{
 		if (UNiagaraEmitter* Emitter = Cast<UNiagaraEmitter>(AssetData.GetAsset()))
 		{
-			FText Tooltip = LOCTEXT("NiagaraEditorDropEmitterTooltip", "Add emitter to system");
-			OutTooltipText = Tooltip.ToString();
+			if (AssetReferenceFilter)
+			{
+				FText FailureReason;
+				if (!AssetReferenceFilter->PassesFilter(AssetData, &FailureReason))
+				{
+					if (OutTooltipText.IsEmpty())
+					{
+						OutTooltipText = FailureReason.ToString();
+					}
+					continue;
+				}
+			}
+
+			OutTooltipText = LOCTEXT("NiagaraEditorDropEmitterTooltip", "Add emitter to system").ToString();
 			OutOkIcon = true;
 			break;
 		}
 	}
 }
 
+TSharedPtr<IAssetReferenceFilter> UEdGraphSchema_NiagaraSystemOverview::MakeAssetReferenceFilter(const UEdGraph* Graph)
+{
+	if (Graph)
+	{
+		if (const UNiagaraSystem* NiagaraSystem = Graph->GetTypedOuter<UNiagaraSystem>())
+		{
+			if (GEditor)
+			{
+				FAssetReferenceFilterContext AssetReferenceFilterContext;
+				AssetReferenceFilterContext.AddReferencingAsset(NiagaraSystem);
+				return GEditor->MakeAssetReferenceFilter(AssetReferenceFilterContext);
+			}
+		}
+	}
+
+	return {};
+}
+
 #undef LOCTEXT_NAMESPACE
 
diff --git a/Engine/Plugins/FX/Niagara/Source/NiagaraEditor/Private/ViewModels/NiagaraScratchPadScriptViewModel.cpp b/Engine/Plugins/FX/Niagara/Source/NiagaraEditor/Private/ViewModels/NiagaraScratchPadScriptViewModel.cpp
--- a/Engine/Plugins/FX/Niagara/Source/NiagaraEditor/Private/ViewModels/NiagaraScratchPadScriptViewModel.cpp
+++ b/Engine/Plugins/FX/Niagara/Source/NiagaraEditor/Private/ViewModels/NiagaraScratchPadScriptViewModel.cpp
@@ -12,6 +12,8 @@
 #include "ViewModels/NiagaraScriptGraphViewModel.h"
 #include "ViewModels/NiagaraSystemViewModel.h"
 
+#include "Editor.h"
+#include "Editor/EditorEngine.h"
 #include "FileHelpers.h"
 #include "ScopedTransaction.h"
 #include "Framework/Commands/UICommandList.h"
@@ -86,6 +88,18 @@ void FNiagaraScratchPadScriptViewModel::Initialize(UNiagaraScript* Script, UNiag
 	ScriptHierarchyViewModel->OnHierarchyPropertiesChanged().AddSP(this, &FNiagaraScratchPadScriptViewModel::OnHierarchyPropertiesChanged);
 }
 
+TSharedPtr<IAssetReferenceFilter> FNiagaraScratchPadScriptViewModel::MakeAssetReferenceFilter() const
+{
+	if (OriginalScript && GEditor)
+	{
+		FAssetReferenceFilterContext AssetReferenceFilterContext;
+		AssetReferenceFilterContext.AddReferencingAsset(OriginalScript);
+		return GEditor->MakeAssetReferenceFilter(AssetReferenceFilterContext);
+	}
+
+	return {};
+}
+
 bool FNiagaraScratchPadScriptViewModel::IsValid() const
 {
 	return EditScript.Script != nullptr && GraphViewModel->GetScriptSource() != nullptr;
diff --git a/Engine/Plugins/FX/Niagara/Source/NiagaraEditor/Private/ViewModels/NiagaraStandaloneScriptViewModel.cpp b/Engine/Plugins/FX/Niagara/Source/NiagaraEditor/Private/ViewModels/NiagaraStandaloneScriptViewModel.cpp
--- a/Engine/Plugins/FX/Niagara/Source/NiagaraEditor/Private/ViewModels/NiagaraStandaloneScriptViewModel.cpp
+++ b/Engine/Plugins/FX/Niagara/Source/NiagaraEditor/Private/ViewModels/NiagaraStandaloneScriptViewModel.cpp
@@ -8,6 +8,8 @@
 #include "NiagaraScriptSource.h"
 #include "NiagaraParameterDefinitions.h"
 
+#include "Editor.h"
+#include "Editor/EditorEngine.h"
 
 FNiagaraStandaloneScriptViewModel::FNiagaraStandaloneScriptViewModel(
 	FText DisplayName,
@@ -30,6 +32,18 @@ void FNiagaraStandaloneScriptViewModel::Initialize(FVersionedNiagaraScript& InSc
 	SendLastCompileMessages(SourceScript);
 }
 
+TSharedPtr<IAssetReferenceFilter> FNiagaraStandaloneScriptViewModel::MakeAssetReferenceFilter() const
+{
+	if (SourceScript.Script && GEditor)
+	{
+		FAssetReferenceFilterContext AssetReferenceFilterContext;
+		AssetReferenceFilterContext.AddReferencingAsset(SourceScript.Script);
+		return GEditor->MakeAssetReferenceFilter(AssetReferenceFilterContext);
+	}
+
+	return {};
+}
+
 FVersionedNiagaraScript FNiagaraStandaloneScriptViewModel::GetStandaloneScript()
 {
 	checkf(Scripts.Num() == 1, TEXT("StandaloneScriptViewModel did not have exactly one script!"));
diff --git a/Engine/Plugins/FX/Niagara/Source/NiagaraEditor/Private/ViewModels/NiagaraStandaloneScriptViewModel.h b/Engine/Plugins/FX/Niagara/Source/NiagaraEditor/Private/ViewModels/NiagaraStandaloneScriptViewModel.h
--- a/Engine/Plugins/FX/Niagara/Source/NiagaraEditor/Private/ViewModels/NiagaraStandaloneScriptViewModel.h
+++ b/Engine/Plugins/FX/Niagara/Source/NiagaraEditor/Private/ViewModels/NiagaraStandaloneScriptViewModel.h
@@ -24,6 +24,8 @@ public:
 
 	NIAGARAEDITOR_API void Initialize(FVersionedNiagaraScript& InScript, const FVersionedNiagaraScript& InSourceScript);
 
+	TSharedPtr<IAssetReferenceFilter> MakeAssetReferenceFilter() const override;
+
 	//~ Begin INiagaraParameterDefinitionsSubscriberViewModel Interface
 protected:
 	NIAGARAEDITOR_API virtual INiagaraParameterDefinitionsSubscriber* GetParameterDefinitionsSubscriber() override;
diff --git a/Engine/Plugins/FX/Niagara/Source/NiagaraEditor/Private/ViewModels/NiagaraSystemScriptViewModel.cpp b/Engine/Plugins/FX/Niagara/Source/NiagaraEditor/Private/ViewModels/NiagaraSystemScriptViewModel.cpp
--- a/Engine/Plugins/FX/Niagara/Source/NiagaraEditor/Private/ViewModels/NiagaraSystemScriptViewModel.cpp
+++ b/Engine/Plugins/FX/Niagara/Source/NiagaraEditor/Private/ViewModels/NiagaraSystemScriptViewModel.cpp
@@ -11,6 +11,9 @@
 #include "ViewModels/NiagaraScriptViewModel.h"
 #include "ViewModels/NiagaraSystemViewModel.h"
 
+#include "Editor.h"
+#include "Editor/EditorEngine.h"
+
 FNiagaraSystemScriptViewModel::FNiagaraSystemScriptViewModel(bool bInIsForDataProcessingOnly)
 	: FNiagaraScriptViewModel(NSLOCTEXT("SystemScriptViewModel", "GraphName", "System"), ENiagaraParameterEditMode::EditAll, bInIsForDataProcessingOnly)
 {
@@ -36,6 +39,21 @@ FNiagaraSystemScriptViewModel::~FNiagaraSystemScriptViewModel()
 	}
 }
 
+TSharedPtr<IAssetReferenceFilter> FNiagaraSystemScriptViewModel::MakeAssetReferenceFilter() const
+{
+	if (const UNiagaraSystem* System = GetSystem())
+	{
+		if (GEditor)
+		{
+			FAssetReferenceFilterContext AssetReferenceFilterContext;
+			AssetReferenceFilterContext.AddReferencingAsset(System);
+			return GEditor->MakeAssetReferenceFilter(AssetReferenceFilterContext);
+		}
+	}
+
+	return {};
+}
+
 void FNiagaraSystemScriptViewModel::OnSystemVMCompiled(UNiagaraSystem* InSystem)
 {
 	{
diff --git a/Engine/Plugins/FX/Niagara/Source/NiagaraEditor/Private/ViewModels/NiagaraSystemScriptViewModel.h b/Engine/Plugins/FX/Niagara/Source/NiagaraEditor/Private/ViewModels/NiagaraSystemScriptViewModel.h
--- a/Engine/Plugins/FX/Niagara/Source/NiagaraEditor/Private/ViewModels/NiagaraSystemScriptViewModel.h
+++ b/Engine/Plugins/FX/Niagara/Source/NiagaraEditor/Private/ViewModels/NiagaraSystemScriptViewModel.h
@@ -22,6 +22,8 @@ public:
 
 	virtual ~FNiagaraSystemScriptViewModel() override;
 
+	TSharedPtr<IAssetReferenceFilter> MakeAssetReferenceFilter() const override;
+
 	FOnSystemCompiled& OnSystemCompiled();
 
 	void CompileSystem(bool bForce);
diff --git a/Engine/Plugins/FX/Niagara/Source/NiagaraEditor/Public/EdGraphSchema_Niagara.h b/Engine/Plugins/FX/Niagara/Source/NiagaraEditor/Public/EdGraphSchema_Niagara.h
--- a/Engine/Plugins/FX/Niagara/Source/NiagaraEditor/Public/EdGraphSchema_Niagara.h
+++ b/Engine/Plugins/FX/Niagara/Source/NiagaraEditor/Public/EdGraphSchema_Niagara.h
@@ -10,6 +10,7 @@
 
 #include "EdGraphSchema_Niagara.generated.h"
 
+class IAssetReferenceFilter;
 class UEdGraph;
 struct FNiagaraVariable;
 struct FNiagaraVariableInfo;
@@ -201,6 +202,8 @@ private:
 	NIAGARAEDITOR_API void GenerateDataInterfacePinMenu(UToolMenu* ToolMenu, const FName SectionName, const UEdGraphPin* GraphPin, FNiagaraTypeDefinition TypeDef) const;
 
 	static NIAGARAEDITOR_API bool CheckCircularConnection(TSet<const UEdGraphNode*>& VisitedNodes, const UEdGraphNode* InNode, const UEdGraphNode* InTestNode);
+
+	static TSharedPtr<IAssetReferenceFilter> MakeAssetReferenceFilter(const UEdGraph* Graph);
 };
 
 class FNiagaraConnectionDrawingPolicy : public FConnectionDrawingPolicy
diff --git a/Engine/Plugins/FX/Niagara/Source/NiagaraEditor/Public/EdGraphSchema_NiagaraSystemOverview.h b/Engine/Plugins/FX/Niagara/Source/NiagaraEditor/Public/EdGraphSchema_NiagaraSystemOverview.h
--- a/Engine/Plugins/FX/Niagara/Source/NiagaraEditor/Public/EdGraphSchema_NiagaraSystemOverview.h
+++ b/Engine/Plugins/FX/Niagara/Source/NiagaraEditor/Public/EdGraphSchema_NiagaraSystemOverview.h
@@ -7,7 +7,7 @@
 #include "ConnectionDrawingPolicy.h"
 #include "EdGraphSchema_NiagaraSystemOverview.generated.h"
 
-
+class IAssetReferenceFilter;
 
 UCLASS(MinimalAPI)
 class UEdGraphSchema_NiagaraSystemOverview : public UEdGraphSchema
@@ -26,6 +26,8 @@ class UEdGraphSchema_NiagaraSystemOverview : public UEdGraphSchema
 	NIAGARAEDITOR_API virtual void GetAssetsGraphHoverMessage(const TArray<FAssetData>& Assets, const UEdGraph* HoverGraph, FString& OutTooltipText, bool& OutOkIcon) const override;
 
 private:
+	static TSharedPtr<IAssetReferenceFilter> MakeAssetReferenceFilter(const UEdGraph* Graph);
+
 	static NIAGARAEDITOR_API int32 CurrentCacheRefreshID;
 };
 
diff --git a/Engine/Plugins/FX/Niagara/Source/NiagaraEditor/Public/ViewModels/NiagaraScratchPadScriptViewModel.h b/Engine/Plugins/FX/Niagara/Source/NiagaraEditor/Public/ViewModels/NiagaraScratchPadScriptViewModel.h
--- a/Engine/Plugins/FX/Niagara/Source/NiagaraEditor/Public/ViewModels/NiagaraScratchPadScriptViewModel.h
+++ b/Engine/Plugins/FX/Niagara/Source/NiagaraEditor/Public/ViewModels/NiagaraScratchPadScriptViewModel.h
@@ -29,6 +29,8 @@ public:
 
 	void Initialize(UNiagaraScript* Script, UNiagaraScript* InEditScript, TWeakPtr<class FNiagaraSystemViewModel> InSystemViewModel);
 
+	NIAGARAEDITOR_API TSharedPtr<IAssetReferenceFilter> MakeAssetReferenceFilter() const override;
+
 	NIAGARAEDITOR_API bool IsValid() const;
 	
 	NIAGARAEDITOR_API void TransferFromOldWhenDoingApply(TSharedPtr< FNiagaraScratchPadScriptViewModel> InOldScriptVM);
diff --git a/Engine/Plugins/FX/Niagara/Source/NiagaraEditor/Public/ViewModels/NiagaraScriptViewModel.h b/Engine/Plugins/FX/Niagara/Source/NiagaraEditor/Public/ViewModels/NiagaraScriptViewModel.h
--- a/Engine/Plugins/FX/Niagara/Source/NiagaraEditor/Public/ViewModels/NiagaraScriptViewModel.h
+++ b/Engine/Plugins/FX/Niagara/Source/NiagaraEditor/Public/ViewModels/NiagaraScriptViewModel.h
@@ -8,6 +8,7 @@
 #include "EditorUndoClient.h"
 #include "NiagaraParameterDefinitionsSubscriberViewModel.h"
 
+class IAssetReferenceFilter;
 class UNiagaraScript;
 class UNiagaraScriptSource;
 class UNiagaraScriptVariable;
@@ -49,6 +50,8 @@ public:
 
 	void SetScripts(FVersionedNiagaraEmitter InEmitter);
 
+	virtual TSharedPtr<IAssetReferenceFilter> MakeAssetReferenceFilter() const { return {}; }
+
 	/** Gets the view model for the input parameter collection. */
 	TSharedRef<FNiagaraScriptInputCollectionViewModel> GetInputCollectionViewModel();
 	
diff --git a/Engine/Plugins/Mutable/Source/CustomizableObjectEditor/Private/MuCOE/EdGraphSchema_CustomizableObject.cpp b/Engine/Plugins/Mutable/Source/CustomizableObjectEditor/Private/MuCOE/EdGraphSchema_CustomizableObject.cpp
--- a/Engine/Plugins/Mutable/Source/CustomizableObjectEditor/Private/MuCOE/EdGraphSchema_CustomizableObject.cpp
+++ b/Engine/Plugins/Mutable/Source/CustomizableObjectEditor/Private/MuCOE/EdGraphSchema_CustomizableObject.cpp
@@ -8,12 +8,16 @@
 #include "CustomizableObjectConnectionDrawingPolicy.h"
 #include "CustomizableObjectEditorSettings.h"
 #include "EdGraphNode_Comment.h"
+#include "Editor.h"
+#include "Editor/AssetReferenceFilter.h"
+#include "Editor/EditorEngine.h"
 #include "Engine/SkeletalMesh.h"
 #include "Engine/StaticMesh.h"
 #include "Framework/Application/SlateApplication.h"
 #include "Framework/Commands/GenericCommands.h"
 #include "GraphEditorActions.h"
 #include "Materials/MaterialInterface.h"
+#include "MuCO/CustomizableObject.h"
 #include "MuCO/CustomizableObjectExtension.h"
 #include "MuCO/ICustomizableObjectModule.h"
 #include "MuCOE/CustomizableObjectGraph.h"
@@ -1170,8 +1174,14 @@ void UEdGraphSchema_CustomizableObject::DroppedAssetsOnGraph(const TArray<FAsset
 	const int PixelOffset = 20;
 	int CurrentOffset = 0;
 	
-	for (FAssetData Asset : Assets)
+	const TSharedPtr<IAssetReferenceFilter> AssetReferenceFilter = MakeAssetReferenceFilter(Graph);
+	for (const FAssetData& Asset : Assets)
 	{
+		if (AssetReferenceFilter.IsValid() && !AssetReferenceFilter->PassesFilter(Asset))
+		{
+			continue;
+		}
+
 		// If it is not a valid asset to be spawned then just skip it
 		ESpawnableObjectType ObjectType;
 		if (!IsSpawnableAsset(Asset,ObjectType))
@@ -1281,10 +1291,10 @@ FConnectionDrawingPolicy* UEdGraphSchema_CustomizableObject::CreateConnectionDra
 void UEdGraphSchema_CustomizableObject::GetAssetsGraphHoverMessage(const TArray<FAssetData>& Assets,
                                                                    const UEdGraph* HoverGraph, FString& OutTooltipText, bool& OutOkIcon) const
 {
-	Super::GetAssetsGraphHoverMessage(Assets, HoverGraph, OutTooltipText, OutOkIcon);
+	OutTooltipText.Reset();
+	OutOkIcon = false;
 
-	// Accept entry by default
-	OutOkIcon = true;
+	const TSharedPtr<IAssetReferenceFilter> AssetReferenceFilter = MakeAssetReferenceFilter(HoverGraph);
 	uint32 AmountOfIncompatibleAssets = 0;
 
 	// Iterate over the assets
@@ -1292,10 +1302,28 @@ void UEdGraphSchema_CustomizableObject::GetAssetsGraphHoverMessage(const TArray<
 	{
 		// On first fail abort the consequent checks and tell the user
 		ESpawnableObjectType ObjectType;
-		if (!IsSpawnableAsset(Asset, ObjectType))
+		if (IsSpawnableAsset(Asset, ObjectType))
+		{
+			if (AssetReferenceFilter)
+			{
+				FText FailureReason;
+				if (!AssetReferenceFilter->PassesFilter(Asset, &FailureReason))
+				{
+					if (OutTooltipText.IsEmpty())
+					{
+						OutTooltipText = FailureReason.ToString();
+					}
+					continue;
+				}
+			}
+
+			OutTooltipText.Reset();
+			OutOkIcon = true;
+			break;
+		}
+		else
 		{
 			AmountOfIncompatibleAssets++;
-			OutOkIcon = false;
 
 			// Stop checking once we know that more than one asset is not compatible, the UI output will be the same
 			if (AmountOfIncompatibleAssets > 1)
@@ -1306,21 +1334,21 @@ void UEdGraphSchema_CustomizableObject::GetAssetsGraphHoverMessage(const TArray<
 	}
 
 	// Output debug message depending on the quantity of incompatible objects
-	if (!OutOkIcon)
+	if (!OutOkIcon && OutTooltipText.IsEmpty())
 	{
 		if (Assets.Num() == 1)
 		{
-			OutTooltipText = FString("Incompatible asset selected : No node can be created for this type of asset.");
+			OutTooltipText = LOCTEXT("IncompatibleAsset", "Incompatible asset selected : No node can be created for this type of asset.").ToString();
 		}
 		else if (Assets.Num() > 1 )
 		{
 			if (Assets.Num() == AmountOfIncompatibleAssets)
 			{
-				OutTooltipText = FString("Incompatible assets selected : No node can be created for any of the selected assets.");
+				OutTooltipText = LOCTEXT("AllIncompatibleAssets", "Incompatible assets selected : No node can be created for any of the selected assets.").ToString();
 			}
 			else 
 			{
-				OutTooltipText = FString("Incompatible asset selected : Some assets will not be placed as nodes on the graph.");
+				OutTooltipText = LOCTEXT("SomeIncompatibleAssets", "Incompatible asset selected : Some assets will not be placed as nodes on the graph.").ToString();
 			}
 		}
 	}
@@ -1589,5 +1617,22 @@ bool UEdGraphSchema_CustomizableObject::IsPassthrough(const FName& PinCategory)
 		PinCategory == PC_PassthroughSkeletalMesh;
 }
 
+TSharedPtr<IAssetReferenceFilter> UEdGraphSchema_CustomizableObject::MakeAssetReferenceFilter(const UEdGraph* Graph)
+{
+	if (Graph)
+	{
+		if (const UCustomizableObject* CustomizableObject = Cast<const UCustomizableObject>(Graph->GetOuter()))
+		{
+			if (GEditor)
+			{
+				FAssetReferenceFilterContext AssetReferenceFilterContext;
+				AssetReferenceFilterContext.AddReferencingAsset(CustomizableObject);
+				return GEditor->MakeAssetReferenceFilter(AssetReferenceFilterContext);
+			}
+		}
+	}
+
+	return {};
+}
 
 #undef LOCTEXT_NAMESPACE
diff --git a/Engine/Plugins/Mutable/Source/CustomizableObjectEditor/Public/MuCOE/EdGraphSchema_CustomizableObject.h b/Engine/Plugins/Mutable/Source/CustomizableObjectEditor/Public/MuCOE/EdGraphSchema_CustomizableObject.h
--- a/Engine/Plugins/Mutable/Source/CustomizableObjectEditor/Public/MuCOE/EdGraphSchema_CustomizableObject.h
+++ b/Engine/Plugins/Mutable/Source/CustomizableObjectEditor/Public/MuCOE/EdGraphSchema_CustomizableObject.h
@@ -16,6 +16,7 @@ class UEdGraphPin;
 class UObject;
 struct FAssetData;
 struct FEdGraphPinType;
+class IAssetReferenceFilter;
 class ICustomizableObjectEditor;
 
 
@@ -75,6 +76,9 @@ public:
 	static UE_API FName GetPinCategoryName(const FName& PinCategory);
 	static UE_API FText GetPinCategoryFriendlyName(const FName& PinCategory);
 	static UE_API bool IsPassthrough(const FName& PinCategory);
+
+private:
+	static TSharedPtr<IAssetReferenceFilter> MakeAssetReferenceFilter(const UEdGraph* Graph);
 };
 
 #undef UE_API
diff --git a/Engine/Plugins/PCG/Source/PCGEditor/Private/PCGEditorGraphSchema.cpp b/Engine/Plugins/PCG/Source/PCGEditor/Private/PCGEditorGraphSchema.cpp
--- a/Engine/Plugins/PCG/Source/PCGEditor/Private/PCGEditorGraphSchema.cpp
+++ b/Engine/Plugins/PCG/Source/PCGEditor/Private/PCGEditorGraphSchema.cpp
@@ -30,6 +30,9 @@
 #include "AssetRegistry/AssetData.h"
 #include "AssetRegistry/AssetRegistryModule.h"
 #include "Blueprint/BlueprintSupport.h"
+#include "Editor.h"
+#include "Editor/AssetReferenceFilter.h"
+#include "Editor/EditorEngine.h"
 #include "Engine/Blueprint.h"
 #include "Framework/Application/SlateApplication.h"
 #include "ScopedTransaction.h"
@@ -734,6 +737,8 @@ void UPCGEditorGraphSchema::GetDataAssetActions(FGraphActionMenuBuilder& ActionM
 
 void UPCGEditorGraphSchema::DroppedAssetsOnGraph(const TArray<FAssetData>& Assets, const FVector2f& GraphPosition, UEdGraph* Graph) const
 {
+	const TSharedPtr<IAssetReferenceFilter> AssetReferenceFilter = MakeAssetReferenceFilter(Graph);
+
 	FVector2f GraphPositionOffset = GraphPosition;
 	constexpr float PositionOffsetIncrementY = 50.f;
 	UEdGraphPin* NullFromPin = nullptr;
@@ -746,6 +751,11 @@ void UPCGEditorGraphSchema::DroppedAssetsOnGraph(const TArray<FAssetData>& Asset
 
 	for (const FAssetData& AssetData : Assets)
 	{
+		if (AssetReferenceFilter.IsValid() && !AssetReferenceFilter->PassesFilter(AssetData))
+		{
+			continue;
+		}
+
 		if(const UClass* AssetClass = AssetData.GetClass())
 		{
 			if(AssetClass->IsChildOf(UPCGGraphInterface::StaticClass()))
@@ -800,8 +810,13 @@ void UPCGEditorGraphSchema::DroppedAssetsOnGraph(const TArray<FAssetData>& Asset
 	}
 }
 
-void UPCGEditorGraphSchema::GetAssetsGraphHoverMessage(const TArray<FAssetData>& Assets, const UEdGraph* /*HoverGraph*/, FString& OutTooltipText, bool& OutOkIcon) const
+void UPCGEditorGraphSchema::GetAssetsGraphHoverMessage(const TArray<FAssetData>& Assets, const UEdGraph* HoverGraph, FString& OutTooltipText, bool& OutOkIcon) const
 {
+	OutTooltipText.Reset();
+	OutOkIcon = false;
+
+	const TSharedPtr<IAssetReferenceFilter> AssetReferenceFilter = MakeAssetReferenceFilter(HoverGraph);
+
 	for (const FAssetData& AssetData : Assets)
 	{
 		// TODO: get class from asset data, shouldn't require loading
@@ -809,20 +824,37 @@ void UPCGEditorGraphSchema::GetAssetsGraphHoverMessage(const TArray<FAssetData>&
 		{
 			if(AssetClass->IsChildOf(UPCGGraphInterface::StaticClass()) || AssetClass->IsChildOf(UPCGSettings::StaticClass()) || AssetClass->IsChildOf(UPCGDataAsset::StaticClass()) || PCGEditorUtils::IsAssetPCGBlueprint(AssetData))
 			{
+				if (AssetReferenceFilter)
+				{
+					FText FailureReason;
+					if (!AssetReferenceFilter->PassesFilter(AssetData, &FailureReason))
+					{
+						if (OutTooltipText.IsEmpty())
+						{
+							OutTooltipText = FailureReason.ToString();
+						}
+						continue;
+					}
+				}
+
+				OutTooltipText.Reset();
 				OutOkIcon = true;
 				return;
 			}
 			else if(AssetClass->IsChildOf(UBlueprint::StaticClass()))
 			{
-				OutTooltipText = LOCTEXT("PCGEditorDropAssetInvalidBP", "Blueprint does not derive from UPCGBlueprintBaseElement").ToString();
-				OutOkIcon = false;
-				return;
+				if (OutTooltipText.IsEmpty())
+				{
+					OutTooltipText = LOCTEXT("PCGEditorDropAssetInvalidBP", "Blueprint does not derive from UPCGBlueprintBaseElement").ToString();
+				}
 			}
 		}
 	}
 
-	OutTooltipText = LOCTEXT("PCGEditorDropAssetInvalid", "Can't create a node for this asset").ToString();
-	OutOkIcon = false;
+	if (OutTooltipText.IsEmpty())
+	{
+		OutTooltipText = LOCTEXT("PCGEditorDropAssetInvalid", "Can't create a node for this asset").ToString();
+	}
 }
 
 void UPCGEditorGraphSchema::OnPinConnectionDoubleCicked(UEdGraphPin* PinA, UEdGraphPin* PinB, const FVector2f& GraphPosition) const
@@ -1019,6 +1051,24 @@ FLinearColor UPCGEditorGraphSchema::GetPinColor(const UEdGraphPin* InPin) const
 	return PCGPin ? FPCGModule::GetConstDataTypeRegistry().GetPinColor(PCGPin->GetCurrentTypesID()) : GetDefault<UPCGEditorSettings>()->DefaultPinColor;
 }
 
+TSharedPtr<IAssetReferenceFilter> UPCGEditorGraphSchema::MakeAssetReferenceFilter(const UEdGraph* Graph)
+{
+	if (const UPCGEditorGraph* EditorGraph = Cast<const UPCGEditorGraph>(Graph))
+	{
+		if (const UPCGGraph* PCGGraph = EditorGraph->GetPCGGraph())
+		{
+			if (GEditor)
+			{
+				FAssetReferenceFilterContext AssetReferenceFilterContext;
+				AssetReferenceFilterContext.AddReferencingAsset(PCGGraph);
+				return GEditor->MakeAssetReferenceFilter(AssetReferenceFilterContext);
+			}
+		}
+	}
+
+	return {};
+}
+
 FPCGEditorConnectionDrawingPolicy::FPCGEditorConnectionDrawingPolicy(int32 InBackLayerID, int32 InFrontLayerID, float InZoomFactor, const FSlateRect& InClippingRect, FSlateWindowElementList& InDrawElements, UEdGraph* InGraph)
 	: FConnectionDrawingPolicy(InBackLayerID, InFrontLayerID, InZoomFactor, InClippingRect, InDrawElements)
 	, Graph(CastChecked<UPCGEditorGraph>(InGraph))
diff --git a/Engine/Plugins/PCG/Source/PCGEditor/Public/Schema/PCGEditorGraphSchema.h b/Engine/Plugins/PCG/Source/PCGEditor/Public/Schema/PCGEditorGraphSchema.h
--- a/Engine/Plugins/PCG/Source/PCGEditor/Public/Schema/PCGEditorGraphSchema.h
+++ b/Engine/Plugins/PCG/Source/PCGEditor/Public/Schema/PCGEditorGraphSchema.h
@@ -14,6 +14,7 @@
 
 enum class EPCGElementType : uint8;
 
+class IAssetReferenceFilter;
 class UPCGEditorGraph;
 
 struct FPCGActionsFilter
@@ -75,6 +76,9 @@ protected:
 	PCGEDITOR_API virtual void GetDataAssetActions(FGraphActionMenuBuilder& ActionMenuBuilder, const FPCGActionsFilter& ActionFilter) const;
 
 	PCGEDITOR_API virtual bool TryCreateConnectionInternal(UEdGraphPin* A, UEdGraphPin* B, bool bAddConversionNodeIfNeeded) const;
+
+private:
+	static TSharedPtr<IAssetReferenceFilter> MakeAssetReferenceFilter(const UEdGraph* Graph);
 };
 
 class FPCGEditorConnectionDrawingPolicy : public FConnectionDrawingPolicy
diff --git a/Engine/Plugins/Runtime/Metasound/Source/MetasoundEditor/Private/MetasoundEditorGraphSchema.cpp b/Engine/Plugins/Runtime/Metasound/Source/MetasoundEditor/Private/MetasoundEditorGraphSchema.cpp
--- a/Engine/Plugins/Runtime/Metasound/Source/MetasoundEditor/Private/MetasoundEditorGraphSchema.cpp
+++ b/Engine/Plugins/Runtime/Metasound/Source/MetasoundEditor/Private/MetasoundEditorGraphSchema.cpp
@@ -6,6 +6,8 @@
 #include "EdGraphNode_Comment.h"
 #include "EdGraphSchema_K2.h"
 #include "Editor.h"
+#include "Editor/AssetReferenceFilter.h"
+#include "Editor/EditorEngine.h"
 #include "Engine/Selection.h"
 #include "Framework/Commands/GenericCommands.h"
 #include "Framework/MultiBox/MultiBoxBuilder.h"
@@ -2386,7 +2388,7 @@ void UMetasoundEditorGraphSchema::GetAssetsGraphHoverMessage(const TArray<FAsset
 		return;
 	}
 
-	OutTooltipText = TEXT("Add MetaSound reference to Graph.");
+	OutTooltipText = LOCTEXT("AddMetaSoundToGraph", "Add MetaSound reference to Graph.").ToString();
 
 	const UMetasoundEditorGraph* MetaSoundGraph = CastChecked<UMetasoundEditorGraph>(HoverGraph);
 	const UObject& MetaSound = MetaSoundGraph->GetMetasoundChecked();
@@ -2394,6 +2396,8 @@ void UMetasoundEditorGraphSchema::GetAssetsGraphHoverMessage(const TArray<FAsset
 	const FMetasoundAssetBase* MetaSoundAsset = IMetasoundUObjectRegistry::Get().GetObjectAsAssetBase(&MetaSound);
 	check(MetaSoundAsset);
 
+	const TSharedPtr<IAssetReferenceFilter> AssetReferenceFilter = MakeAssetReferenceFilter(MetaSoundGraph);
+
 	for (const FAssetData& Data : Assets)
 	{
 		UClass* AssetClass = Data.GetClass();
@@ -2402,7 +2406,7 @@ void UMetasoundEditorGraphSchema::GetAssetsGraphHoverMessage(const TArray<FAsset
 		if (!bIsRegisteredClass)
 		{
 			OutOkIcon = false;
-			OutTooltipText = TEXT("Asset(s) must all be MetaSounds.");
+			OutTooltipText = LOCTEXT("NotAllMetaSounds", "Asset(s) must all be MetaSounds.").ToString();
 			break;
 		}
 
@@ -2412,7 +2416,7 @@ void UMetasoundEditorGraphSchema::GetAssetsGraphHoverMessage(const TArray<FAsset
 			if (!DroppedMetaSoundAsset)
 			{
 				OutOkIcon = false;
-				OutTooltipText = TEXT("Asset(s) not (all) MetaSound(s) that can be referenced.");
+				OutTooltipText = LOCTEXT("NotAllReferenceableMetaSounds", "Asset(s) not (all) MetaSound(s) that can be referenced.").ToString();
 				break;
 			}
 
@@ -2420,21 +2424,29 @@ void UMetasoundEditorGraphSchema::GetAssetsGraphHoverMessage(const TArray<FAsset
 			if (!EnumHasAnyFlags(DroppedAccessFlags, EMetasoundFrontendClassAccessFlags::Referenceable))
 			{
 				OutOkIcon = false;
-				OutTooltipText = TEXT("Asset(s) do not contain access flag 'Reference' (see 'Access Flags'). Cannot reference provided MetaSound Asset.");
+				OutTooltipText = LOCTEXT("NotAllAccessibleMetaSounds", "Asset(s) do not contain access flag 'Reference' (see 'Access Flags'). Cannot reference provided MetaSound Asset.").ToString();
 				break;
 			}
 
 			if (MetaSoundAsset->AddingReferenceCausesLoop(*DroppedMetaSoundAsset))
 			{
 				OutOkIcon = false;
-				OutTooltipText = TEXT("Cannot add asset(s) that would create a reference loop/loops.");
+				OutTooltipText = LOCTEXT("AddingMetaSoundWouldCauseLoop", "Cannot add asset(s) that would create a reference loop/loops.").ToString();
+				break;
+			}
+
+			FText FailureReason;
+			if (AssetReferenceFilter.IsValid() && !AssetReferenceFilter->PassesFilter(Data, &FailureReason))
+			{
+				OutOkIcon = false;
+				OutTooltipText = FailureReason.ToString();
 				break;
 			}
 		}
 		else
 		{
 			OutOkIcon = false;
-			OutTooltipText = TEXT("Asset(s) not found.");
+			OutTooltipText = LOCTEXT("AssetNotFound", "Asset(s) not found.").ToString();
 			break;
 		}
 	}
@@ -2457,28 +2469,37 @@ void UMetasoundEditorGraphSchema::GetAssetsPinHoverMessage(const TArray<FAssetDa
 					const bool bAssetTypesMatch = SchemaPrivate::DataTypeSupportsAssetTypes(RegistryInfo, Assets);
 					if (bAssetTypesMatch)
 					{
-						OutTooltipText = FString::Format(TEXT("Set to '{0}'"), { *Assets[0].AssetName.ToString() });
+						const TSharedPtr<IAssetReferenceFilter> AssetReferenceFilter = MakeAssetReferenceFilter(MetaSoundNode->GetGraph());
+						FText FailureReason;
+						if (AssetReferenceFilter.IsValid() && !AssetReferenceFilter->PassesFilter(Assets[0], &FailureReason))
+						{
+							OutTooltipText = FailureReason.ToString();
+							OutOkIcon = false;
+							return;
+						}
+
+						OutTooltipText = FText::FormatOrdered(LOCTEXT("DropAsset", "Set to '{0}'"), FText::FromName(Assets[0].AssetName)).ToString();
 						OutOkIcon = true;
 						return;
 					}
 
-					OutTooltipText = FString::Format(TEXT("'{0}': Invalid Type"), { *Assets[0].AssetName.ToString() });
+					OutTooltipText = FText::FormatOrdered(LOCTEXT("InvalidType", "'{0}': Invalid Type"), FText::FromName(Assets[0].AssetName)).ToString();
 					OutOkIcon = false;
 					return;
 				}
 
-				OutTooltipText = TEXT("Cannot drop multiple assets on single pin.");
+				OutTooltipText = LOCTEXT("CannotDropMultiple", "Cannot drop multiple assets on single pin.").ToString();
 				OutOkIcon = false;
 				return;
 			}
 
-			OutTooltipText = FString::Format(TEXT("Node '{0}' does not support drag/drop"), { *Node->GetName() });
+			OutTooltipText = FText::FormatOrdered(LOCTEXT("DragDropNotSupported", "Node '{0}' does not support drag/drop"), FText::FromString(Node->GetName())).ToString();
 			OutOkIcon = false;
 			return;
 		}
 	}
 
-	OutTooltipText = FString();
+	OutTooltipText.Reset();
 	OutOkIcon = false;
 }
 
@@ -2500,6 +2521,8 @@ void UMetasoundEditorGraphSchema::DroppedAssetsOnGraph(const TArray<FAssetData>&
 	UMetaSoundBuilderBase& Builder = FDocumentBuilderRegistry::GetChecked().FindOrBeginBuilding(MetaSound);
 	EMetaSoundBuilderResult Result = EMetaSoundBuilderResult::Failed;
 
+	const TSharedPtr<IAssetReferenceFilter> AssetReferenceFilter = MakeAssetReferenceFilter(MetaSoundGraph);
+
 	for (const FAssetData& DroppedAsset : Assets)
 	{
 		if (UObject* DroppedObject = DroppedAsset.GetAsset())
@@ -2521,6 +2544,11 @@ void UMetasoundEditorGraphSchema::DroppedAssetsOnGraph(const TArray<FAssetData>&
 				continue;
 			}
 
+			if (AssetReferenceFilter.IsValid() && !AssetReferenceFilter->PassesFilter(DroppedAsset))
+			{
+				continue;
+			}
+
 			if (!bModifiedObjects)
 			{
 				MetaSound.Modify();
@@ -2544,11 +2572,6 @@ void UMetasoundEditorGraphSchema::DroppedAssetsOnGraph(const TArray<FAssetData>&
 	}
 }
 
-void UMetasoundEditorGraphSchema::DroppedAssetsOnNode(const TArray<FAssetData>& Assets, const FVector2f& GraphPosition, UEdGraphNode* Node) const
-{
-	// Still needed?
-}
-
 void UMetasoundEditorGraphSchema::DroppedAssetsOnPin(const TArray<FAssetData>& Assets, const FVector2f& GraphPosition, UEdGraphPin* Pin) const
 {
 	using namespace Metasound::Editor;
@@ -2567,15 +2590,22 @@ void UMetasoundEditorGraphSchema::DroppedAssetsOnPin(const TArray<FAssetData>& A
 			const bool bAssetTypesMatch = SchemaPrivate::DataTypeSupportsAssetTypes(RegistryInfo, Assets);
 			if (bAssetTypesMatch)
 			{
-				UObject* Object = Assets.Last().GetAsset();
-				const FText TransactionText = FText::Format(LOCTEXT("ChangeDefaultObjectTransaction", "Set {0} to '{1}'"),
-					Pin->GetDisplayName(),
-					FText::FromName(Object->GetFName()));
-				const FScopedTransaction Transaction(TransactionText);
-				Node->Modify();
-
-				constexpr bool bMarkAsModified = true;
-				TrySetDefaultObject(*Pin, Object, bMarkAsModified);
+				const TSharedPtr<IAssetReferenceFilter> AssetReferenceFilter = MakeAssetReferenceFilter(Node->GetGraph());
+				if (!AssetReferenceFilter.IsValid() || AssetReferenceFilter->PassesFilter(Assets.Last()))
+				{
+					UObject* Object = Assets.Last().GetAsset();
+					if (Object)
+					{
+						const FText TransactionText = FText::Format(LOCTEXT("ChangeDefaultObjectTransaction", "Set {0} to '{1}'"),
+							Pin->GetDisplayName(),
+							FText::FromName(Object->GetFName()));
+						const FScopedTransaction Transaction(TransactionText);
+						Node->Modify();
+
+						constexpr bool bMarkAsModified = true;
+						TrySetDefaultObject(*Pin, Object, bMarkAsModified);
+					}
+				}
 			}
 		}
 	}
@@ -2942,5 +2972,24 @@ void UMetasoundEditorGraphSchema::QueryNodeClasses(TFunctionRef<void(TUniquePtr<
 		OnClassFound(MoveTemp(Result));
 	}
 }
+
+TSharedPtr<IAssetReferenceFilter> UMetasoundEditorGraphSchema::MakeAssetReferenceFilter(const UEdGraph* Graph)
+{
+	if (const UMetasoundEditorGraph* MetasoundGraph = Cast<const UMetasoundEditorGraph>(Graph))
+	{
+		if (const UObject* Metasound = MetasoundGraph->GetMetasound())
+		{
+			if (GEditor)
+			{
+				FAssetReferenceFilterContext AssetReferenceFilterContext;
+				AssetReferenceFilterContext.AddReferencingAsset(Metasound);
+				return GEditor->MakeAssetReferenceFilter(AssetReferenceFilterContext);
+			}
+		}
+	}
+
+	return {};
+}
+
 #undef LOCTEXT_NAMESPACE
 
diff --git a/Engine/Plugins/Runtime/Metasound/Source/MetasoundEditor/Public/MetasoundEditorGraphSchema.h b/Engine/Plugins/Runtime/Metasound/Source/MetasoundEditor/Public/MetasoundEditorGraphSchema.h
--- a/Engine/Plugins/Runtime/Metasound/Source/MetasoundEditor/Public/MetasoundEditorGraphSchema.h
+++ b/Engine/Plugins/Runtime/Metasound/Source/MetasoundEditor/Public/MetasoundEditorGraphSchema.h
@@ -17,6 +17,7 @@
 
 
 // Forward Declarations
+class IAssetReferenceFilter;
 class UEdGraph;
 class UEdGraphNode;
 class UMetaSoundPatch;
@@ -556,7 +557,6 @@ public:
 	virtual void GetAssetsGraphHoverMessage(const TArray<FAssetData>& Assets, const UEdGraph* HoverGraph, FString& OutTooltipText, bool& OutOkIcon) const override;
 	virtual void GetAssetsPinHoverMessage(const TArray<FAssetData>& Assets, const UEdGraphPin* HoverPin, FString& OutTooltipText, bool& OutOkIcon) const override;
 	virtual void DroppedAssetsOnGraph(const TArray<FAssetData>& Assets, const FVector2f& GraphPosition, UEdGraph* Graph) const override;
-	virtual void DroppedAssetsOnNode(const TArray<FAssetData>& Assets, const FVector2f& GraphPosition, UEdGraphNode* Node) const override;
 	virtual void DroppedAssetsOnPin(const TArray<FAssetData>& Assets, const FVector2f& GraphPosition, UEdGraphPin* Pin) const override;
 	virtual int32 GetNodeSelectionCount(const UEdGraph* Graph) const override;
 	virtual TSharedPtr<FEdGraphSchemaAction> GetCreateCommentAction() const override;
@@ -581,6 +581,8 @@ private:
 
 	/** Adds action for creating a comment */
 	void GetCommentAction(FGraphActionMenuBuilder& ActionMenuBuilder, const UEdGraph* CurrentGraph = nullptr) const;
+
+	static TSharedPtr<IAssetReferenceFilter> MakeAssetReferenceFilter(const UEdGraph* Graph);
 };
 
 #undef UE_API
diff --git a/Engine/Plugins/TextureGraph/Source/TextureGraphEditor/Private/EdGraph/TG_EdGraphSchema.cpp b/Engine/Plugins/TextureGraph/Source/TextureGraphEditor/Private/EdGraph/TG_EdGraphSchema.cpp
--- a/Engine/Plugins/TextureGraph/Source/TextureGraphEditor/Private/EdGraph/TG_EdGraphSchema.cpp
+++ b/Engine/Plugins/TextureGraph/Source/TextureGraphEditor/Private/EdGraph/TG_EdGraphSchema.cpp
@@ -9,6 +9,8 @@
 
 #include "ScopedTransaction.h"
 #include "2D/TextureHelper.h"
+#include "Editor/AssetReferenceFilter.h"
+#include "TG_Editor.h"
 #include "TG_Graph.h"
 #include "UObject/UObjectIterator.h"
 #include "MaterialValueType.h"
@@ -388,9 +390,15 @@ void UTG_EdGraphSchema::DroppedAssetsOnGraph(const TArray<FAssetData>& Assets, c
 
 	FVector2f ExpressionPosition = GraphPosition;
 
-	for (int32 AssetIdx = 0; AssetIdx < Assets.Num(); ++AssetIdx)
+	const TSharedPtr<IAssetReferenceFilter> AssetReferenceFilter = MakeAssetReferenceFilter(Graph);
+	for (const FAssetData& AssetData : Assets)
 	{
-		UObject* Asset = Assets[AssetIdx].GetAsset();
+		if (AssetReferenceFilter.IsValid() && !AssetReferenceFilter->PassesFilter(AssetData))
+		{
+			continue;
+		}
+
+		UObject* Asset = AssetData.GetAsset();
 		UEdGraphNode* ExpressionNode = nullptr;
 		
 		auto Tex = Cast<UTexture>(Asset);
@@ -444,57 +452,86 @@ void UTG_EdGraphSchema::DroppedAssetsOnNode(const TArray<FAssetData>& Assets, co
 
 	if(EdNode)
 	{
-		UObject* Asset = Assets.Last().GetAsset();
+		const FAssetData& AssetData = Assets.Last();
+		UObject* Asset = AssetData.GetAsset();
 		UTG_Expression* Expression = EdNode->GetNode()->GetExpression();
 
 		if(Expression->CanHandleAsset(Asset))
 		{
-			const FScopedTransaction Transaction(NSLOCTEXT("UTG_EdGraphSchema", "DroppedAssetsOnNode", "Dropped Assets OnNode"));
-			Expression->Modify();
+			const TSharedPtr<IAssetReferenceFilter> AssetReferenceFilter = MakeAssetReferenceFilter(EdNode->GetGraph());
+			if (!AssetReferenceFilter.IsValid() || AssetReferenceFilter->PassesFilter(AssetData))
+			{
+				const FScopedTransaction Transaction(NSLOCTEXT("UTG_EdGraphSchema", "DroppedAssetsOnNode", "Dropped Assets OnNode"));
+				Expression->Modify();
 
-			Expression->SetAsset(Asset);
+				Expression->SetAsset(Asset);
+			}
 		}
 	}
 }
 
 void UTG_EdGraphSchema::GetAssetsGraphHoverMessage(const TArray<FAssetData>& Assets, const UEdGraph* HoverGraph,
-                                                   FString& OutTooltipText, bool& OutOkIcon) const
+	FString& OutTooltipText, bool& OutOkIcon) const
 {
+	OutTooltipText.Reset();
 	OutOkIcon = false;
 
-	for (int32 AssetIdx = 0; AssetIdx < Assets.Num(); ++AssetIdx)
+	const TSharedPtr<IAssetReferenceFilter> AssetReferenceFilter = MakeAssetReferenceFilter(HoverGraph);
+	for (const FAssetData& AssetData : Assets)
 	{
-		UObject* Asset = Assets[AssetIdx].GetAsset();
+		UObject* Asset = AssetData.GetAsset();
 		
-		if (TextureHelper::CanSupportTexture(Cast<UTexture>(Asset)))
-		{
-			OutOkIcon = true;
-		}
-		else if(UMaterialInterface* Material = Cast<UMaterialInterface>(Asset); Material != nullptr)
-		{
-			OutOkIcon = true;
-		}
-		else if(UTextureGraph* TextureGraph = Cast<UTextureGraph>(Asset); TextureGraph != nullptr)
+		if (TextureHelper::CanSupportTexture(Cast<UTexture>(Asset)) ||
+			Cast<UMaterialInterface>(Asset) != nullptr ||
+			Cast<UMaterialFunctionInterface>(Asset) != nullptr ||
+			Cast<UTextureGraph>(Asset) != nullptr)
 		{
+			if (AssetReferenceFilter)
+			{
+				FText FailureReason;
+				if (!AssetReferenceFilter->PassesFilter(AssetData, &FailureReason))
+				{
+					if (OutTooltipText.IsEmpty())
+					{
+						OutTooltipText = FailureReason.ToString();
+					}
+					continue;
+				}
+			}
+
+			OutTooltipText.Reset();
 			OutOkIcon = true;
+			return;
 		}
 	}
 }
 
 void UTG_EdGraphSchema::GetAssetsNodeHoverMessage(const TArray<FAssetData>& Assets, const UEdGraphNode* HoverNode, FString& OutTooltipText, bool& OutOkIcon) const
 {
+	OutTooltipText.Reset();
 	OutOkIcon = false;
 
 	const UTG_EdGraphNode* EdNode = Cast<UTG_EdGraphNode>(HoverNode);
 	
 	if(EdNode)
 	{
-		UObject* Asset = Assets.Last().GetAsset();
+		const FAssetData& AssetData = Assets.Last();
+		UObject* Asset = AssetData.GetAsset();
 		UTG_Expression* Expression = EdNode->GetNode()->GetExpression();
 		
 		if (Expression->CanHandleAsset(Asset))
 		{
-			OutOkIcon = true;
+			const TSharedPtr<IAssetReferenceFilter> AssetReferenceFilter = MakeAssetReferenceFilter(EdNode->GetGraph());
+			FText FailureReason;
+			if (AssetReferenceFilter.IsValid() && !AssetReferenceFilter->PassesFilter(AssetData, &FailureReason))
+			{
+				OutTooltipText = FailureReason.ToString();
+			}
+			else
+			{
+				OutTooltipText.Reset();
+				OutOkIcon = true;
+			}
 		}
 	}
 }
@@ -552,4 +589,17 @@ FLinearColor UTG_EdGraphSchema::GetCategoryColor(FName Category)
 	return UTG_EdGraphSchema::TGNodeColor;
 }
 
+TSharedPtr<IAssetReferenceFilter> UTG_EdGraphSchema::MakeAssetReferenceFilter(const UEdGraph* Graph)
+{
+	if (const UTG_EdGraph* EdNode = Cast<const UTG_EdGraph>(Graph))
+	{
+		if (TSharedPtr<FTG_Editor> TGEditor = EdNode->TGEditor.Pin())
+		{
+			return TGEditor->MakeAssetReferenceFilter();
+		}
+	}
+
+	return {};
+}
+
 #undef LOCTEXT_NAMESPACE
diff --git a/Engine/Plugins/TextureGraph/Source/TextureGraphEditor/Private/EdGraph/TG_EdGraphSchema.h b/Engine/Plugins/TextureGraph/Source/TextureGraphEditor/Private/EdGraph/TG_EdGraphSchema.h
--- a/Engine/Plugins/TextureGraph/Source/TextureGraphEditor/Private/EdGraph/TG_EdGraphSchema.h
+++ b/Engine/Plugins/TextureGraph/Source/TextureGraphEditor/Private/EdGraph/TG_EdGraphSchema.h
@@ -7,6 +7,7 @@
 #include "EdGraph/EdGraphSchema.h"
 #include "TG_EdGraphSchema.generated.h"
 
+class IAssetReferenceFilter;
 class UEdGraph;
 class UTG_Node;
 class UTG_Pin;
@@ -95,5 +96,7 @@ private:
 
 	// Internal check that 2 pins are compatible for connection 
 	bool ArePinsCompatible_Internal(const UEdGraphPin* InA, const UEdGraphPin* InB) const;
+
+	static TSharedPtr<IAssetReferenceFilter> MakeAssetReferenceFilter(const UEdGraph* Graph);
 };
 
diff --git a/Engine/Plugins/TextureGraph/Source/TextureGraphEditor/Private/TG_Editor.cpp b/Engine/Plugins/TextureGraph/Source/TextureGraphEditor/Private/TG_Editor.cpp
--- a/Engine/Plugins/TextureGraph/Source/TextureGraphEditor/Private/TG_Editor.cpp
+++ b/Engine/Plugins/TextureGraph/Source/TextureGraphEditor/Private/TG_Editor.cpp
@@ -61,6 +61,7 @@
 #include "EdGraph/TG_EditorErrorReporter.h"
 #include "Misc/UObjectToken.h"
 #include "EdGraph/TG_EdGraphSchemaActions.h"
+#include "Editor/EditorEngine.h"
 #include "Editor/UnrealEdEngine.h"
 #include "Expressions/Input/TG_Expression_Graph.h"
 #include "Model/Mix/MixInterface.h"
@@ -435,6 +436,18 @@ FText FTG_Editor::GetOriginalObjectName() const
 	return FText::FromString(GetEditingObjects()[0]->GetName());
 }
 
+TSharedPtr<IAssetReferenceFilter> FTG_Editor::MakeAssetReferenceFilter() const
+{
+	if (GEditor)
+	{
+		FAssetReferenceFilterContext AssetReferenceFilterContext;
+		AssetReferenceFilterContext.AddReferencingAssets(GetEditingObjects());
+		return GEditor->MakeAssetReferenceFilter(AssetReferenceFilterContext);
+	}
+
+	return {};
+}
+
 UTG_Expression* FTG_Editor::CreateNewExpression(UClass* NewExpressionClass, const FVector2D& NodePos, bool bAutoSelect, bool bAutoAssignResource, const class UEdGraph* Graph /*= nullptr*/)
 {
 	check(NewExpressionClass->IsChildOf(UTG_Expression::StaticClass()));
diff --git a/Engine/Plugins/TextureGraph/Source/TextureGraphEditor/Private/TG_Editor.h b/Engine/Plugins/TextureGraph/Source/TextureGraphEditor/Private/TG_Editor.h
--- a/Engine/Plugins/TextureGraph/Source/TextureGraphEditor/Private/TG_Editor.h
+++ b/Engine/Plugins/TextureGraph/Source/TextureGraphEditor/Private/TG_Editor.h
@@ -79,6 +79,7 @@ public:
 
 	virtual FText									GetOriginalObjectName() const override;
 
+	virtual TSharedPtr<IAssetReferenceFilter>		MakeAssetReferenceFilter() const override;
 
 	UTG_Expression*									CreateNewExpression(UClass* NewExpressionClass, const FVector2D& NodePos, bool bAutoSelect, bool bAutoAssignResource, const class UEdGraph* Graph = nullptr) override;
 	
diff --git a/Engine/Plugins/TextureGraph/Source/TextureGraphEditor/Private/TG_InstanceEditor.cpp b/Engine/Plugins/TextureGraph/Source/TextureGraphEditor/Private/TG_InstanceEditor.cpp
--- a/Engine/Plugins/TextureGraph/Source/TextureGraphEditor/Private/TG_InstanceEditor.cpp
+++ b/Engine/Plugins/TextureGraph/Source/TextureGraphEditor/Private/TG_InstanceEditor.cpp
@@ -7,6 +7,7 @@
 #include "Components/MeshComponent.h"
 #include "Components/StaticMeshComponent.h"
 #include "Editor.h"
+#include "Editor/EditorEngine.h"
 
 #include "Framework/Application/SlateApplication.h"
 #include "PropertyEditorModule.h"
@@ -145,6 +146,18 @@ FText FTG_InstanceEditor::GetOriginalObjectName() const
 	return FText::FromString(GetEditingObjects()[0]->GetName());
 }
 
+TSharedPtr<IAssetReferenceFilter> FTG_InstanceEditor::MakeAssetReferenceFilter() const
+{
+	if (GEditor)
+	{
+		FAssetReferenceFilterContext AssetReferenceFilterContext;
+		AssetReferenceFilterContext.AddReferencingAssets(GetEditingObjects());
+		return GEditor->MakeAssetReferenceFilter(AssetReferenceFilterContext);
+	}
+
+	return {};
+}
+
 void FTG_InstanceEditor::OnClose()
 {
 	// // Add analytics tag
diff --git a/Engine/Plugins/TextureGraph/Source/TextureGraphEditor/Private/TG_InstanceEditor.h b/Engine/Plugins/TextureGraph/Source/TextureGraphEditor/Private/TG_InstanceEditor.h
--- a/Engine/Plugins/TextureGraph/Source/TextureGraphEditor/Private/TG_InstanceEditor.h
+++ b/Engine/Plugins/TextureGraph/Source/TextureGraphEditor/Private/TG_InstanceEditor.h
@@ -79,6 +79,7 @@ public:
 
 	virtual FText									GetOriginalObjectName() const override;
 
+	virtual TSharedPtr<IAssetReferenceFilter>		MakeAssetReferenceFilter() const override;
 
 	// ~Begin FAssetEditorToolkit interface
 	virtual void									OnClose() override;
diff --git a/Engine/Plugins/TextureGraph/Source/TextureGraphEditor/Public/ITG_Editor.h b/Engine/Plugins/TextureGraph/Source/TextureGraphEditor/Public/ITG_Editor.h
--- a/Engine/Plugins/TextureGraph/Source/TextureGraphEditor/Public/ITG_Editor.h
+++ b/Engine/Plugins/TextureGraph/Source/TextureGraphEditor/Public/ITG_Editor.h
@@ -45,6 +45,11 @@ public:
 	 */
 	virtual FText					GetOriginalObjectName() const { return FText::GetEmpty(); }
 
+	/**
+	 * Creates an IAssetReferenceFilter that can be used to check if an asset is valid to be used with this texture graph.
+	 */
+	virtual TSharedPtr<IAssetReferenceFilter> MakeAssetReferenceFilter() const { return {}; }
+
 	/** Gets the TextureGraph edited */
 	virtual class UMixInterface*	GetTextureGraphInterface() const = 0;
 
diff --git a/Engine/Plugins/VirtualProduction/Avalanche/Source/AvalancheMediaEditor/Private/Playback/Graph/AvaPlaybackEditorGraphSchema.cpp b/Engine/Plugins/VirtualProduction/Avalanche/Source/AvalancheMediaEditor/Private/Playback/Graph/AvaPlaybackEditorGraphSchema.cpp
--- a/Engine/Plugins/VirtualProduction/Avalanche/Source/AvalancheMediaEditor/Private/Playback/Graph/AvaPlaybackEditorGraphSchema.cpp
+++ b/Engine/Plugins/VirtualProduction/Avalanche/Source/AvalancheMediaEditor/Private/Playback/Graph/AvaPlaybackEditorGraphSchema.cpp
@@ -214,24 +214,6 @@ void UAvaPlaybackEditorGraphSchema::BreakPinLinks(UEdGraphPin& TargetPin, bool b
 	}
 }
 
-void UAvaPlaybackEditorGraphSchema::GetAssetsGraphHoverMessage(const TArray<FAssetData>& Assets, const UEdGraph* HoverGraph
-	, FString& OutTooltipText, bool& OutOkIcon) const
-{
-	Super::GetAssetsGraphHoverMessage(Assets, HoverGraph, OutTooltipText, OutOkIcon);
-}
-
-void UAvaPlaybackEditorGraphSchema::DroppedAssetsOnGraph(const TArray<FAssetData>& Assets, const FVector2f& GraphPosition
-	, UEdGraph* Graph) const
-{
-	Super::DroppedAssetsOnGraph(Assets, GraphPosition, Graph);
-}
-
-void UAvaPlaybackEditorGraphSchema::DroppedAssetsOnNode(const TArray<FAssetData>& Assets, const FVector2f& GraphPosition
-	, UEdGraphNode* Node) const
-{
-	Super::DroppedAssetsOnNode(Assets, GraphPosition, Node);
-}
-
 int32 UAvaPlaybackEditorGraphSchema::GetNodeSelectionCount(const UEdGraph* Graph) const
 {
 	return Super::GetNodeSelectionCount(Graph);
diff --git a/Engine/Plugins/VirtualProduction/Avalanche/Source/AvalancheMediaEditor/Private/Playback/Graph/AvaPlaybackEditorGraphSchema.h b/Engine/Plugins/VirtualProduction/Avalanche/Source/AvalancheMediaEditor/Private/Playback/Graph/AvaPlaybackEditorGraphSchema.h
--- a/Engine/Plugins/VirtualProduction/Avalanche/Source/AvalancheMediaEditor/Private/Playback/Graph/AvaPlaybackEditorGraphSchema.h
+++ b/Engine/Plugins/VirtualProduction/Avalanche/Source/AvalancheMediaEditor/Private/Playback/Graph/AvaPlaybackEditorGraphSchema.h
@@ -55,9 +55,6 @@ public:
 	virtual bool ShouldHidePinDefaultValue(UEdGraphPin* Pin) const override;
 	virtual void BreakNodeLinks(UEdGraphNode& TargetNode) const override;
 	virtual void BreakPinLinks(UEdGraphPin& TargetPin, bool bSendsNodeNotifcation) const override;
-	virtual void GetAssetsGraphHoverMessage(const TArray<FAssetData>& Assets, const UEdGraph* HoverGraph, FString& OutTooltipText, bool& OutOkIcon) const override;
-	virtual void DroppedAssetsOnGraph(const TArray<struct FAssetData>& Assets, const FVector2f& GraphPosition, UEdGraph* Graph) const override;
-	virtual void DroppedAssetsOnNode(const TArray<FAssetData>& Assets, const FVector2f& GraphPosition, UEdGraphNode* Node) const override;
 	virtual int32 GetNodeSelectionCount(const UEdGraph* Graph) const override;
 	virtual TSharedPtr<FEdGraphSchemaAction> GetCreateCommentAction() const override;
 	//~ End UEdGraphSchema
diff --git a/Engine/Source/Editor/AnimGraph/Private/AnimationGraphSchema.cpp b/Engine/Source/Editor/AnimGraph/Private/AnimationGraphSchema.cpp
--- a/Engine/Source/Editor/AnimGraph/Private/AnimationGraphSchema.cpp
+++ b/Engine/Source/Editor/AnimGraph/Private/AnimationGraphSchema.cpp
@@ -12,6 +12,9 @@
 #include "ObjectEditorUtils.h"
 #include "K2Node.h"
 #include "EdGraphSchema_K2_Actions.h"
+#include "Editor.h"
+#include "Editor/AssetReferenceFilter.h"
+#include "Editor/EditorEngine.h"
 #include "Kismet2/BlueprintEditorUtils.h"
 #include "Animation/AnimSequence.h"
 #include "AnimStateNode.h"
@@ -679,15 +682,23 @@ void UAnimationGraphSchema::UpdateNodeWithAsset(UK2Node* K2Node, UAnimationAsset
 
 void UAnimationGraphSchema::DroppedAssetsOnGraph( const TArray<FAssetData>& Assets, const FVector2f& GraphPosition, UEdGraph* Graph ) const 
 {
-	if (Graph != NULL)
+	if (Graph != nullptr)
 	{
 		if (UAnimationAsset* AnimationAsset = FAssetData::GetFirstAsset<UAnimationAsset>(Assets))
 		{
-			SpawnNodeFromAsset(AnimationAsset, GraphPosition, Graph, NULL);
+			const TSharedPtr<IAssetReferenceFilter> AssetReferenceFilter = MakeAssetReferenceFilter(Graph);
+			if (!AssetReferenceFilter.IsValid() || AssetReferenceFilter->PassesFilter(FAssetData(AnimationAsset)))
+			{
+				SpawnNodeFromAsset(AnimationAsset, GraphPosition, Graph, nullptr);
+			}
 		}
 		else if (UPhysicsAsset* PhysicsAsset = FAssetData::GetFirstAsset<UPhysicsAsset>(Assets))
 		{
-			SpawnRigidBodyNodeFromAsset(PhysicsAsset, FDeprecateSlateVector2D(GraphPosition), Graph);
+			const TSharedPtr<IAssetReferenceFilter> AssetReferenceFilter = MakeAssetReferenceFilter(Graph);
+			if (!AssetReferenceFilter.IsValid() || AssetReferenceFilter->PassesFilter(FAssetData(PhysicsAsset)))
+			{
+				SpawnRigidBodyNodeFromAsset(PhysicsAsset, FDeprecateSlateVector2D(GraphPosition), Graph);
+			}
 		}
 	}
 }
@@ -696,29 +707,37 @@ void UAnimationGraphSchema::DroppedAssetsOnNode(const TArray<FAssetData>& Assets
 {
 	UAnimationAsset* Asset = FAssetData::GetFirstAsset<UAnimationAsset>(Assets);
 	UK2Node* K2Node = Cast<UK2Node>(Node);
-	if ((Asset != NULL) && (K2Node!= NULL))
+	if ((Asset != nullptr) && (K2Node!= nullptr))
 	{
-		UpdateNodeWithAsset(K2Node, Asset);
+		const TSharedPtr<IAssetReferenceFilter> AssetReferenceFilter = MakeAssetReferenceFilter(K2Node->GetGraph());
+		if (!AssetReferenceFilter.IsValid() || AssetReferenceFilter->PassesFilter(FAssetData(Asset)))
+		{
+			UpdateNodeWithAsset(K2Node, Asset);
+		}
 	}
 }
 
 void UAnimationGraphSchema::DroppedAssetsOnPin(const TArray<FAssetData>& Assets, const FVector2f& GraphPosition, UEdGraphPin* Pin) const
 {
 	UAnimationAsset* Asset = FAssetData::GetFirstAsset<UAnimationAsset>(Assets);
-	if ((Asset != NULL) && (Pin != NULL))
+	if ((Asset != nullptr) && (Pin != nullptr))
 	{
-		// Don't have access to bounding information for node, using fixed offset that should work for most cases.
-		const FVector2f FixedOffset(-250.0f, 50.0f);
-		SpawnNodeFromAsset(Asset, GraphPosition + FixedOffset, Pin->GetOwningNode()->GetGraph(), Pin);
+		const TSharedPtr<IAssetReferenceFilter> AssetReferenceFilter = MakeAssetReferenceFilter(Pin->GetOwningNode()->GetGraph());
+		if (!AssetReferenceFilter.IsValid() || AssetReferenceFilter->PassesFilter(FAssetData(Asset)))
+		{
+			// Don't have access to bounding information for node, using fixed offset that should work for most cases.
+			const FVector2f FixedOffset(-250.0f, 50.0f);
+			SpawnNodeFromAsset(Asset, GraphPosition + FixedOffset, Pin->GetOwningNode()->GetGraph(), Pin);
+		}
 	}
 }
 
-void UAnimationGraphSchema::GetAssetsNodeHoverMessage(const TArray<FAssetData>& Assets, const UEdGraphNode* HoverNode, FString& OutTooltipText, bool& OutOkIcon) const 
+void UAnimationGraphSchema::GetAssetsNodeHoverMessage(const TArray<FAssetData>& Assets, const UEdGraphNode* HoverNode, FString& OutTooltipText, bool& OutOkIcon) const
 { 
 	UAnimationAsset* Asset = FAssetData::GetFirstAsset<UAnimationAsset>(Assets);
-	if ((Asset == NULL) || (HoverNode == NULL) || !HoverNode->IsA(UAnimGraphNode_Base::StaticClass()))
+	if ((Asset == nullptr) || (HoverNode == nullptr) || !HoverNode->IsA(UAnimGraphNode_Base::StaticClass()))
 	{
-		OutTooltipText = TEXT("");
+		OutTooltipText.Reset();
 		OutOkIcon = false;
 		return;
 	}
@@ -736,8 +755,18 @@ void UAnimationGraphSchema::GetAssetsNodeHoverMessage(const TArray<FAssetData>&
 	}
 	else if (bCanPlayAsset)
 	{
-		OutOkIcon = true;
-		OutTooltipText = FText::Format(LOCTEXT("AssetNodeHoverMessage_Success", "Change node to play '{0}'"), FText::FromString(Asset->GetName())).ToString();
+		const TSharedPtr<IAssetReferenceFilter> AssetReferenceFilter = MakeAssetReferenceFilter(HoverNode->GetGraph());
+		FText FailureReason;
+		if (AssetReferenceFilter.IsValid() && !AssetReferenceFilter->PassesFilter(FAssetData(Asset), &FailureReason))
+		{
+			OutOkIcon = false;
+			OutTooltipText = FailureReason.ToString();
+		}
+		else
+		{
+			OutOkIcon = true;
+			OutTooltipText = FText::Format(LOCTEXT("AssetNodeHoverMessage_Success", "Change node to play '{0}'"), FText::FromString(Asset->GetName())).ToString();
+		}
 	}
 	else
 	{
@@ -749,9 +778,9 @@ void UAnimationGraphSchema::GetAssetsNodeHoverMessage(const TArray<FAssetData>&
 void UAnimationGraphSchema::GetAssetsPinHoverMessage(const TArray<FAssetData>& Assets, const UEdGraphPin* HoverPin, FString& OutTooltipText, bool& OutOkIcon) const 
 { 
 	UAnimationAsset* Asset = FAssetData::GetFirstAsset<UAnimationAsset>(Assets);
-	if ((Asset == NULL) || (HoverPin == NULL))
+	if ((Asset == nullptr) || (HoverPin == nullptr))
 	{
-		OutTooltipText = TEXT("");
+		OutTooltipText.Reset();
 		OutOkIcon = false;
 		return;
 	}
@@ -765,8 +794,18 @@ void UAnimationGraphSchema::GetAssetsPinHoverMessage(const TArray<FAssetData>& A
 
 	if (bSkelMatch && bTypeMatch && bDirectionMatch)
 	{
-		OutOkIcon = true;
-		OutTooltipText = FText::Format(LOCTEXT("AssetPinHoverMessage_Success", "Play {0} and feed to {1}"), FText::FromString(Asset->GetName()), FText::FromName(HoverPin->PinName)).ToString();
+		const TSharedPtr<IAssetReferenceFilter> AssetReferenceFilter = MakeAssetReferenceFilter(HoverPin->GetOwningNode()->GetGraph());
+		FText FailureReason;
+		if (AssetReferenceFilter.IsValid() && !AssetReferenceFilter->PassesFilter(FAssetData(Asset), &FailureReason))
+		{
+			OutOkIcon = false;
+			OutTooltipText = FailureReason.ToString();
+		}
+		else
+		{
+			OutOkIcon = true;
+			OutTooltipText = FText::Format(LOCTEXT("AssetPinHoverMessage_Success", "Play {0} and feed to {1}"), FText::FromString(Asset->GetName()), FText::FromName(HoverPin->PinName)).ToString();
+		}
 	}
 	else
 	{
@@ -793,21 +832,41 @@ void UAnimationGraphSchema::GetAssetsGraphHoverMessage(const TArray<FAssetData>&
 				OutTooltipText = LOCTEXT("SkeletonsNotCompatible", "Skeletons are not compatible").ToString();
 			}
 		}
-		else if(UAnimMontage* Montage = FAssetData::GetFirstAsset<UAnimMontage>(Assets))
+		else if(UAnimMontage* Montage = Cast<UAnimMontage>(AnimationAsset))
 		{
 			OutOkIcon = false;
 			OutTooltipText = LOCTEXT("NoMontagesInAnimGraphs", "Montages cannot be used in animation graphs").ToString();
 		}
 		else
 		{
-			OutOkIcon = true;
-			OutTooltipText = TEXT("");
+			const TSharedPtr<IAssetReferenceFilter> AssetReferenceFilter = MakeAssetReferenceFilter(HoverGraph);
+			FText FailureReason;
+			if (AssetReferenceFilter.IsValid() && !AssetReferenceFilter->PassesFilter(FAssetData(AnimationAsset), &FailureReason))
+			{
+				OutOkIcon = false;
+				OutTooltipText = FailureReason.ToString();
+			}
+			else
+			{
+				OutOkIcon = true;
+				OutTooltipText.Reset();
+			}
 		}
 	}
 	else if(UPhysicsAsset* PhysicsAsset = FAssetData::GetFirstAsset<UPhysicsAsset>(Assets))
 	{
-		OutOkIcon = true;
-		OutTooltipText = TEXT("");
+		const TSharedPtr<IAssetReferenceFilter> AssetReferenceFilter = MakeAssetReferenceFilter(HoverGraph);
+		FText FailureReason;
+		if (AssetReferenceFilter.IsValid() && !AssetReferenceFilter->PassesFilter(FAssetData(PhysicsAsset), &FailureReason))
+		{
+			OutOkIcon = false;
+			OutTooltipText = FailureReason.ToString();
+		}
+		else
+		{
+			OutOkIcon = true;
+			OutTooltipText.Reset();
+		}
 	}
 }
 
@@ -1226,4 +1285,19 @@ FVector2D UAnimationGraphSchema::GetPositionForNewLinkedInputPoseNode(UEdGraph&
 	}
 }
 
+TSharedPtr<IAssetReferenceFilter> UAnimationGraphSchema::MakeAssetReferenceFilter(const UEdGraph* Graph)
+{
+	if (const UBlueprint* Blueprint = FBlueprintEditorUtils::FindBlueprintForGraph(Graph))
+	{
+		if (GEditor)
+		{
+			FAssetReferenceFilterContext AssetReferenceFilterContext;
+			AssetReferenceFilterContext.AddReferencingAsset(Blueprint);
+			return GEditor->MakeAssetReferenceFilter(AssetReferenceFilterContext);
+		}
+	}
+
+	return {};
+}
+
 #undef LOCTEXT_NAMESPACE
diff --git a/Engine/Source/Editor/AnimGraph/Private/AnimationStateMachineSchema.cpp b/Engine/Source/Editor/AnimGraph/Private/AnimationStateMachineSchema.cpp
--- a/Engine/Source/Editor/AnimGraph/Private/AnimationStateMachineSchema.cpp
+++ b/Engine/Source/Editor/AnimGraph/Private/AnimationStateMachineSchema.cpp
@@ -10,6 +10,9 @@
 #include "Framework/MultiBox/MultiBoxBuilder.h"
 #include "ToolMenus.h"
 #include "EdGraph/EdGraph.h"
+#include "Editor.h"
+#include "Editor/AssetReferenceFilter.h"
+#include "Editor/EditorEngine.h"
 #include "Kismet2/BlueprintEditorUtils.h"
 #include "Settings/EditorStyleSettings.h"
 #include "EdGraphNode_Comment.h"
@@ -474,17 +477,21 @@ void UAnimationStateMachineSchema::GetGraphDisplayInformation(const UEdGraph& Gr
 void UAnimationStateMachineSchema::DroppedAssetsOnGraph(const TArray<FAssetData>& Assets, const FVector2f& GraphPosition, UEdGraph* Graph) const
 {
 	UAnimationAsset* Asset = FAssetData::GetFirstAsset<UAnimationAsset>(Assets);
-	if(Asset != NULL && GetNodeClassForAsset(Asset->GetClass()))
+	if(Asset != nullptr && GetNodeClassForAsset(Asset->GetClass()))
 	{
-		// Spawn new state
-		UAnimStateNode* NewStateNode = FEdGraphSchemaAction_NewStateNode::SpawnNodeFromTemplate<UAnimStateNode>(Graph, NewObject<UAnimStateNode>(), GraphPosition);
+		const TSharedPtr<IAssetReferenceFilter> AssetReferenceFilter = MakeAssetReferenceFilter(Graph);
+		if (!AssetReferenceFilter.IsValid() || AssetReferenceFilter->PassesFilter(FAssetData(Asset)))
+		{
+			// Spawn new state
+			UAnimStateNode* NewStateNode = FEdGraphSchemaAction_NewStateNode::SpawnNodeFromTemplate<UAnimStateNode>(Graph, NewObject<UAnimStateNode>(), GraphPosition);
 
-		// Try to name the state close to the asset
-		FEdGraphUtilities::RenameGraphToNameOrCloseToName(NewStateNode->BoundGraph, Asset->GetName());
+			// Try to name the state close to the asset
+			FEdGraphUtilities::RenameGraphToNameOrCloseToName(NewStateNode->BoundGraph, Asset->GetName());
 
-		// Change the current graph context to the inner graph, so that the rest of the drag drop happens inside it
-		FVector2f NewGraphPosition = FVector2f(-300.0f, 0.0f);
-		UAnimationGraphSchema::SpawnNodeFromAsset(Asset, NewGraphPosition, NewStateNode->BoundGraph, NewStateNode->GetPoseSinkPinInsideState());
+			// Change the current graph context to the inner graph, so that the rest of the drag drop happens inside it
+			FVector2f NewGraphPosition = FVector2f(-300.0f, 0.0f);
+			UAnimationGraphSchema::SpawnNodeFromAsset(Asset, NewGraphPosition, NewStateNode->BoundGraph, NewStateNode->GetPoseSinkPinInsideState());
+		}
 	}
 }
 
@@ -492,70 +499,103 @@ void UAnimationStateMachineSchema::DroppedAssetsOnNode(const TArray<FAssetData>&
 {
 	UAnimationAsset* Asset = FAssetData::GetFirstAsset<UAnimationAsset>(Assets);
 	UAnimStateNode* StateNodeUnderCursor = Cast<UAnimStateNode>(Node);
-	if(Asset != NULL && StateNodeUnderCursor != NULL)
+	if(Asset != nullptr && StateNodeUnderCursor != nullptr)
 	{
-		// Dropped onto a state machine state; try some user friendly resposnes
-		if (UEdGraphPin* PosePin = StateNodeUnderCursor->GetPoseSinkPinInsideState())
+		const TSharedPtr<IAssetReferenceFilter> AssetReferenceFilter = MakeAssetReferenceFilter(StateNodeUnderCursor->GetGraph());
+		if (!AssetReferenceFilter.IsValid() || AssetReferenceFilter->PassesFilter(FAssetData(Asset)))
 		{
-			if (PosePin->LinkedTo.Num() > 0)
+			// Dropped onto a state machine state; try some user friendly resposnes
+			if (UEdGraphPin* PosePin = StateNodeUnderCursor->GetPoseSinkPinInsideState())
 			{
-				//@TODO: A2REMOVAL: This doesn't do anything
-				/*
-				// Try dropping in onto the node attached to the sink inside the state
-				check(PosePin->LinkedTo[0] != NULL);
-				UA2Node* A2Node = Cast<UA2Node>(PosePin->LinkedTo[0]->GetOwningNode());
-				if(A2Node != NULL)
+				if (PosePin->LinkedTo.Num() > 0)
 				{
-					UAnimationGraphSchema::UpdateNodeWithAsset(A2Node, Asset);
+					//@TODO: A2REMOVAL: This doesn't do anything
+					/*
+					// Try dropping in onto the node attached to the sink inside the state
+					check(PosePin->LinkedTo[0] != NULL);
+					UA2Node* A2Node = Cast<UA2Node>(PosePin->LinkedTo[0]->GetOwningNode());
+					if(A2Node != NULL)
+					{
+						UAnimationGraphSchema::UpdateNodeWithAsset(A2Node, Asset);
+					}
+					*/
+				}
+				else
+				{
+					// Try dropping onto the pose pin, since nothing is currently connected
+					const FVector2f NewGraphPosition(-300.0f, 0.0f);
+					UAnimationGraphSchema::SpawnNodeFromAsset(Asset, NewGraphPosition, StateNodeUnderCursor->BoundGraph, PosePin);
 				}
-				*/
-			}
-			else
-			{
-				// Try dropping onto the pose pin, since nothing is currently connected
-				const FVector2f NewGraphPosition(-300.0f, 0.0f);
-				UAnimationGraphSchema::SpawnNodeFromAsset(Asset, NewGraphPosition, StateNodeUnderCursor->BoundGraph, PosePin);
 			}
 		}
 	}
 }
 
-void UAnimationStateMachineSchema::DroppedAssetsOnPin(const TArray<FAssetData>& Assets, const FVector2f& GraphPosition, UEdGraphPin* Pin) const
+void UAnimationStateMachineSchema::GetAssetsGraphHoverMessage(const TArray<FAssetData>& Assets, const UEdGraph* HoverGraph, FString& OutTooltipText, bool& OutOkIcon) const
 {
-	// unused for state machines?
+	UAnimationAsset* Asset = FAssetData::GetFirstAsset<UAnimationAsset>(Assets);
+	if ((Asset == nullptr) || (HoverGraph == nullptr))
+	{
+		OutTooltipText.Reset();
+		OutOkIcon = false;
+		return;
+	}
+
+	if (GetNodeClassForAsset(Asset->GetClass()) != nullptr)
+	{
+		const TSharedPtr<IAssetReferenceFilter> AssetReferenceFilter = MakeAssetReferenceFilter(HoverGraph);
+		FText FailureReason;
+		if (AssetReferenceFilter.IsValid() && !AssetReferenceFilter->PassesFilter(FAssetData(Asset), &FailureReason))
+		{
+			OutOkIcon = false;
+			OutTooltipText = FailureReason.ToString();
+		}
+		else
+		{
+			OutOkIcon = true;
+			OutTooltipText.Reset();
+		}
+	}
+	else
+	{
+		OutTooltipText.Reset();
+		OutOkIcon = false;
+	}
 }
 
 void UAnimationStateMachineSchema::GetAssetsNodeHoverMessage(const TArray<FAssetData>& Assets, const UEdGraphNode* HoverNode, FString& OutTooltipText, bool& OutOkIcon) const 
 { 
 	UAnimationAsset* Asset = FAssetData::GetFirstAsset<UAnimationAsset>(Assets);
-	if ((Asset == NULL) || (HoverNode == NULL))
+	if ((Asset == nullptr) || (HoverNode == nullptr))
 	{
-		OutTooltipText = TEXT("");
+		OutTooltipText.Reset();
 		OutOkIcon = false;
 		return;
 	}
 
 	const UAnimStateNode* StateNodeUnderCursor = Cast<const UAnimStateNode>(HoverNode);
-	if (StateNodeUnderCursor != NULL)
+	if (StateNodeUnderCursor != nullptr)
 	{
-		OutOkIcon = true;
-		OutTooltipText = FString::Printf(TEXT("Change node to play %s"), *(Asset->GetName()));
+		const TSharedPtr<IAssetReferenceFilter> AssetReferenceFilter = MakeAssetReferenceFilter(HoverNode->GetGraph());
+		FText FailureReason;
+		if (AssetReferenceFilter.IsValid() && !AssetReferenceFilter->PassesFilter(FAssetData(Asset), &FailureReason))
+		{
+			OutOkIcon = false;
+			OutTooltipText = FailureReason.ToString();
+		}
+		else
+		{
+			OutOkIcon = true;
+			OutTooltipText = FText::FormatOrdered(LOCTEXT("DragDropAssetMessage", "Change node to play {0}"), FText::FromName(Asset->GetFName())).ToString();
+		}
 	}
 	else
 	{
-		OutTooltipText = TEXT("");
+		OutTooltipText.Reset();
 		OutOkIcon = false;
 	}
 }
 
-void UAnimationStateMachineSchema::GetAssetsPinHoverMessage(const TArray<FAssetData>& Assets, const UEdGraphPin* HoverPin, FString& OutTooltipText, bool& OutOkIcon) const 
-{ 
-	// unused for state machines?
-
-	OutTooltipText = TEXT("");
-	OutOkIcon = false;
-}
-
 void UAnimationStateMachineSchema::BreakNodeLinks(UEdGraphNode& TargetNode) const
 {
 	const FScopedTransaction Transaction(NSLOCTEXT("UnrealEd", "GraphEd_BreakNodeLinks", "Break Node Links"));
@@ -582,4 +622,19 @@ void UAnimationStateMachineSchema::BreakSinglePinLink(UEdGraphPin* SourcePin, UE
 	FBlueprintEditorUtils::MarkBlueprintAsModified(Blueprint);
 }
 
+TSharedPtr<IAssetReferenceFilter> UAnimationStateMachineSchema::MakeAssetReferenceFilter(const UEdGraph* Graph)
+{
+	if (const UBlueprint* Blueprint = FBlueprintEditorUtils::FindBlueprintForGraph(Graph))
+	{
+		if (GEditor)
+		{
+			FAssetReferenceFilterContext AssetReferenceFilterContext;
+			AssetReferenceFilterContext.AddReferencingAsset(Blueprint);
+			return GEditor->MakeAssetReferenceFilter(AssetReferenceFilterContext);
+		}
+	}
+
+	return {};
+}
+
 #undef LOCTEXT_NAMESPACE
diff --git a/Engine/Source/Editor/AnimGraph/Public/AnimationGraphSchema.h b/Engine/Source/Editor/AnimGraph/Public/AnimationGraphSchema.h
--- a/Engine/Source/Editor/AnimGraph/Public/AnimationGraphSchema.h
+++ b/Engine/Source/Editor/AnimGraph/Public/AnimationGraphSchema.h
@@ -11,6 +11,7 @@
 #include "AnimationGraphSchema.generated.h"
 
 class FMenuBuilder;
+class IAssetReferenceFilter;
 class UAnimationAsset;
 class UAnimBlueprint;
 class UPhysicsAsset;
@@ -127,6 +128,8 @@ class UAnimationGraphSchema : public UEdGraphSchema_K2
 	/** Unexpose pins that are unused on a linked anim graph node */
 	ANIMGRAPH_API static void HideUnboundPropertyPins(UAnimGraphNode_LinkedAnimGraphBase* Node);
 
+private:
+	static TSharedPtr<IAssetReferenceFilter> MakeAssetReferenceFilter(const UEdGraph* Graph);
 };
 
 
diff --git a/Engine/Source/Editor/AnimGraph/Public/AnimationStateMachineSchema.h b/Engine/Source/Editor/AnimGraph/Public/AnimationStateMachineSchema.h
--- a/Engine/Source/Editor/AnimGraph/Public/AnimationStateMachineSchema.h
+++ b/Engine/Source/Editor/AnimGraph/Public/AnimationStateMachineSchema.h
@@ -12,6 +12,7 @@
 #define UE_API ANIMGRAPH_API
 
 class FMenuBuilder;
+class IAssetReferenceFilter;
 class UEdGraph;
 
 /** Action to add a node to the graph */
@@ -87,14 +88,16 @@ class UAnimationStateMachineSchema : public UEdGraphSchema
 	virtual void GetGraphDisplayInformation(const UEdGraph& Graph, /*out*/ FGraphDisplayInfo& DisplayInfo) const override;
 	virtual void DroppedAssetsOnGraph(const TArray<FAssetData>& Assets, const FVector2f& GraphPosition, UEdGraph* Graph) const override;
 	virtual void DroppedAssetsOnNode(const TArray<FAssetData>& Assets, const FVector2f& GraphPosition, UEdGraphNode* Node) const override;
-	virtual void DroppedAssetsOnPin(const TArray<FAssetData>& Assets, const FVector2f& GraphPosition, UEdGraphPin* Pin) const override;
+	virtual void GetAssetsGraphHoverMessage(const TArray<FAssetData>& Assets, const UEdGraph* HoverGraph, FString& OutTooltipText, bool& OutOkIcon) const override;
 	virtual void GetAssetsNodeHoverMessage(const TArray<FAssetData>& Assets, const UEdGraphNode* HoverNode, FString& OutTooltipText, bool& OutOkIcon) const override;
-	virtual void GetAssetsPinHoverMessage(const TArray<FAssetData>& Assets, const UEdGraphPin* HoverPin, FString& OutTooltipText, bool& OutOkIcon) const override;
 
 	virtual void BreakNodeLinks(UEdGraphNode& TargetNode) const override;
 	virtual void BreakPinLinks(UEdGraphPin& TargetPin, bool bSendsNodeNotification) const override;
 	virtual void BreakSinglePinLink(UEdGraphPin* SourcePin, UEdGraphPin* TargetPin) const override;
 	//~ End UEdGraphSchema Interface
+
+private:
+	static TSharedPtr<IAssetReferenceFilter> MakeAssetReferenceFilter(const UEdGraph* Graph);
 };
 
 #undef UE_API
diff --git a/Engine/Source/Editor/AudioEditor/Classes/SoundCueGraph/SoundCueGraphSchema.h b/Engine/Source/Editor/AudioEditor/Classes/SoundCueGraph/SoundCueGraphSchema.h
--- a/Engine/Source/Editor/AudioEditor/Classes/SoundCueGraph/SoundCueGraphSchema.h
+++ b/Engine/Source/Editor/AudioEditor/Classes/SoundCueGraph/SoundCueGraphSchema.h
@@ -18,6 +18,7 @@
 #define UE_API AUDIOEDITOR_API
 
 class FString;
+class IAssetReferenceFilter;
 class UClass;
 class UEdGraph;
 class UEdGraphNode;
@@ -157,6 +158,8 @@ private:
 	void GetAllSoundNodeActions(FGraphActionMenuBuilder& ActionMenuBuilder, bool bShowSelectedActions) const;
 	/** Adds action for creating a comment */
 	void GetCommentAction(FGraphActionMenuBuilder& ActionMenuBuilder, const UEdGraph* CurrentGraph = NULL) const;
+
+	static TSharedPtr<IAssetReferenceFilter> MakeAssetReferenceFilter(const UEdGraph* Graph);
 	
 	/** Soundnodes that can be selected in this graph */
 	TArray<TSubclassOf<USoundNode>> AllowedSoundNodes;
diff --git a/Engine/Source/Editor/AudioEditor/Private/SoundCueGraphSchema.cpp b/Engine/Source/Editor/AudioEditor/Private/SoundCueGraphSchema.cpp
--- a/Engine/Source/Editor/AudioEditor/Private/SoundCueGraphSchema.cpp
+++ b/Engine/Source/Editor/AudioEditor/Private/SoundCueGraphSchema.cpp
@@ -16,6 +16,7 @@
 #include "EdGraph/EdGraphPin.h"
 #include "EdGraphNode_Comment.h"
 #include "Editor.h"
+#include "Editor/AssetReferenceFilter.h"
 #include "Editor/EditorEngine.h"
 #include "GraphEditor.h"
 #include "GraphEditorActions.h"
@@ -419,31 +420,52 @@ void USoundCueGraphSchema::BreakPinLinks(UEdGraphPin& TargetPin, bool bSendsNode
 
 void USoundCueGraphSchema::GetAssetsGraphHoverMessage(const TArray<FAssetData>& Assets, const UEdGraph* HoverGraph, FString& OutTooltipText, bool& OutOkIcon) const
 {
+	OutTooltipText.Reset();
 	OutOkIcon = false;
 
-	for (int32 AssetIdx = 0; AssetIdx < Assets.Num(); ++AssetIdx)
+	const TSharedPtr<IAssetReferenceFilter> AssetReferenceFilter = MakeAssetReferenceFilter(HoverGraph);
+	for (const FAssetData& AssetData : Assets)
 	{
-		// As soon as one of the items is a sound wave, say we can drag it on... we actually eat only the sound waves.
-		USoundWave* SoundWav = Cast<USoundWave>(Assets[AssetIdx].GetAsset());
-		if (SoundWav)
+		const UObject* Asset = AssetData.GetAsset();
+		// As soon as one of the items is a sound or dialogue wave, say we can drag it on... we actually eat only the sound waves.
+		if (Cast<const USoundWave>(Asset) != nullptr || Cast<const UDialogueWave>(Asset) != nullptr)
 		{
+			if (AssetReferenceFilter)
+			{
+				FText FailureReason;
+				if (!AssetReferenceFilter->PassesFilter(AssetData, &FailureReason))
+				{
+					if (OutTooltipText.IsEmpty())
+					{
+						OutTooltipText = FailureReason.ToString();
+					}
+					continue;
+				}
+			}
+
+			OutTooltipText.Reset();
 			OutOkIcon = true;
-			break;
+			return;
 		}
 	}
 }
 
 void USoundCueGraphSchema::DroppedAssetsOnGraph(const TArray<FAssetData>& Assets, const FVector2f& GraphPosition, UEdGraph* Graph) const
 {
+	const TSharedPtr<IAssetReferenceFilter> AssetReferenceFilter = MakeAssetReferenceFilter(Graph);
+
 	//////////////////////////////////////////////////////////////////////////
 	// Handle dropped USoundWaves
 	TArray<USoundWave*> Waves;
-	for (int32 AssetIdx = 0; AssetIdx < Assets.Num(); ++AssetIdx)
+	for (const FAssetData& AssetData : Assets)
 	{
-		USoundWave* SoundWav = Cast<USoundWave>(Assets[AssetIdx].GetAsset());
-		if (SoundWav)
+		USoundWave* SoundWave = Cast<USoundWave>(AssetData.GetAsset());
+		if (SoundWave)
 		{
-			Waves.Add(SoundWav);
+			if (!AssetReferenceFilter.IsValid() || AssetReferenceFilter->PassesFilter(AssetData))
+			{
+				Waves.Add(SoundWave);
+			}
 		}
 	}
 
@@ -463,12 +485,15 @@ void USoundCueGraphSchema::DroppedAssetsOnGraph(const TArray<FAssetData>& Assets
 	//////////////////////////////////////////////////////////////////////////
 	// Handle dropped UDialogueWaves
 	TArray<UDialogueWave*> Dialogues;
-	for (int32 AssetIdx = 0; AssetIdx < Assets.Num(); ++AssetIdx)
+	for (const FAssetData& AssetData : Assets)
 	{
-		UDialogueWave* DialogueWave = Cast<UDialogueWave>(Assets[AssetIdx].GetAsset());
+		UDialogueWave* DialogueWave = Cast<UDialogueWave>(AssetData.GetAsset());
 		if (DialogueWave)
 		{
-			Dialogues.Add(DialogueWave);
+			if (!AssetReferenceFilter.IsValid() || AssetReferenceFilter->PassesFilter(AssetData))
+			{
+				Dialogues.Add(DialogueWave);
+			}
 		}
 	}
 
@@ -717,4 +742,22 @@ TSharedPtr<FEdGraphSchemaAction> USoundCueGraphSchema::GetCreateCommentAction()
 	return TSharedPtr<FEdGraphSchemaAction>(static_cast<FEdGraphSchemaAction*>(new FSoundCueGraphSchemaAction_NewComment));
 }
 
+TSharedPtr<IAssetReferenceFilter> USoundCueGraphSchema::MakeAssetReferenceFilter(const UEdGraph* Graph)
+{
+	if (const USoundCueGraph* SoundCueGraph = Cast<const USoundCueGraph>(Graph))
+	{
+		if (const USoundCue* SoundCue = SoundCueGraph->GetSoundCue())
+		{
+			if (GEditor)
+			{
+				FAssetReferenceFilterContext AssetReferenceFilterContext;
+				AssetReferenceFilterContext.AddReferencingAsset(SoundCue);
+				return GEditor->MakeAssetReferenceFilter(AssetReferenceFilterContext);
+			}
+		}
+	}
+
+	return {};
+}
+
 #undef LOCTEXT_NAMESPACE
diff --git a/Engine/Source/Editor/AudioEditor/Private/SoundSubmixGraphSchema.cpp b/Engine/Source/Editor/AudioEditor/Private/SoundSubmixGraphSchema.cpp
--- a/Engine/Source/Editor/AudioEditor/Private/SoundSubmixGraphSchema.cpp
+++ b/Engine/Source/Editor/AudioEditor/Private/SoundSubmixGraphSchema.cpp
@@ -147,7 +147,7 @@ bool USoundSubmixGraphSchema::ConnectionCausesLoop(const UEdGraphPin* InputPin,
 void USoundSubmixGraphSchema::GetAssetsGraphHoverMessage(const TArray<FAssetData>& Assets, const UEdGraph* HoverGraph, FString& OutTooltipText, bool& OutOkIcon) const
 {
 	OutOkIcon = true;
-	OutTooltipText = TEXT("Add Submix to Graph.");
+	OutTooltipText = LOCTEXT("DropSoundSubmix", "Add Submix to Graph.").ToString();
 
 	if (HoverGraph)
 	{
@@ -165,7 +165,7 @@ void USoundSubmixGraphSchema::GetAssetsGraphHoverMessage(const TArray<FAssetData
 				if (Assets.ContainsByPredicate(MatchesSubmix))
 				{
 					OutOkIcon = false;
-					OutTooltipText = TEXT("Selected asset or assets already in graph.");
+					OutTooltipText = LOCTEXT("SoundSubmixInGraph", "Selected asset or assets already in graph.").ToString();
 					break;
 				}
 			}
@@ -177,7 +177,7 @@ void USoundSubmixGraphSchema::GetAssetsGraphHoverMessage(const TArray<FAssetData
 		if (!Data.IsInstanceOf(USoundSubmixBase::StaticClass()))
 		{
 			OutOkIcon = false;
-			OutTooltipText = TEXT("Asset(s) must all be Submixes.");
+			OutTooltipText = LOCTEXT("NotASoundSubmix", "Asset(s) must all be Submixes.").ToString();
 			break;
 		}
 	}
diff --git a/Engine/Source/Editor/BlueprintGraph/Classes/EdGraphSchema_K2.h b/Engine/Source/Editor/BlueprintGraph/Classes/EdGraphSchema_K2.h
--- a/Engine/Source/Editor/BlueprintGraph/Classes/EdGraphSchema_K2.h
+++ b/Engine/Source/Editor/BlueprintGraph/Classes/EdGraphSchema_K2.h
@@ -39,6 +39,7 @@ class AActor;
 class FKismetCompilerContext;
 class FMenuBuilder;
 class FProperty;
+class IAssetReferenceFilter;
 class UBlueprint;
 class UClass;
 class UEnum;
@@ -594,7 +595,6 @@ public:
 	UE_API virtual void DroppedAssetsOnGraph(const TArray<FAssetData>& Assets, const FVector2f& GraphPosition, UEdGraph* Graph) const override;
 	UE_API virtual void DroppedAssetsOnNode(const TArray<FAssetData>& Assets, const FVector2f& GraphPosition, UEdGraphNode* Node) const override;
 	UE_API virtual void DroppedAssetsOnPin(const TArray<FAssetData>& Assets, const FVector2f& GraphPosition, UEdGraphPin* Pin) const override;
-	UE_API virtual void GetAssetsNodeHoverMessage(const TArray<FAssetData>& Assets, const UEdGraphNode* HoverNode, FString& OutTooltipText, bool& OutOkIcon) const override;
 	UE_API virtual void GetAssetsPinHoverMessage(const TArray<FAssetData>& Assets, const UEdGraphPin* HoverPin, FString& OutTooltipText, bool& OutOkIcon) const override;
 	UE_API virtual void GetAssetsGraphHoverMessage(const TArray<FAssetData>& Assets, const UEdGraph* HoverGraph, FString& OutTooltipText, bool& OutOkIcon) const override;
 	UE_API virtual bool CanDuplicateGraph(UEdGraph* InSourceGraph) const override;
@@ -1285,6 +1285,8 @@ public:
 	static UE_API void Shutdown();
 private:
 
+	static TSharedPtr<IAssetReferenceFilter> MakeAssetReferenceFilter(const UBlueprint* Blueprint);
+
 	/**
 	 * Returns true if the specified function has any out parameters
 	 * @param [in] Function	The function to check for out parameters
diff --git a/Engine/Source/Editor/BlueprintGraph/Private/EdGraphSchema_K2.cpp b/Engine/Source/Editor/BlueprintGraph/Private/EdGraphSchema_K2.cpp
--- a/Engine/Source/Editor/BlueprintGraph/Private/EdGraphSchema_K2.cpp
+++ b/Engine/Source/Editor/BlueprintGraph/Private/EdGraphSchema_K2.cpp
@@ -5873,6 +5873,9 @@ bool UEdGraphSchema_K2::IsStaticFunctionGraph( const UEdGraph* TestEdGraph ) con
 
 void UEdGraphSchema_K2::DroppedAssetsOnGraph(const TArray<FAssetData>& Assets, const FVector2f& GraphPosition, UEdGraph* Graph) const 
 {
+	UBlueprint* Blueprint = FBlueprintEditorUtils::FindBlueprintForGraph(Graph);
+	TSharedPtr<IAssetReferenceFilter> AssetReferenceFilter = MakeAssetReferenceFilter(Blueprint);
+
 	// only want to spawn event nodes in an event graph
 	if (FBlueprintEditorUtils::IsEventGraph(Graph))
 	{
@@ -5883,12 +5886,14 @@ void UEdGraphSchema_K2::DroppedAssetsOnGraph(const TArray<FAssetData>& Assets, c
 			// if we can find any actions we want to try to spawn the node - only spawning Input Action event nodes currently
 			if (const FAssetBlueprintGraphActions* GraphActions = Module.GetAssetBlueprintGraphActions(AssetData.GetClass()))
 			{
-				GraphActions->TryCreatingAssetNode(AssetData, Graph, FDeprecateSlateVector2D(GraphPosition), EK2NewNodeFlags::SelectNewNode);
+				if (!AssetReferenceFilter || AssetReferenceFilter->PassesFilter(AssetData))
+				{
+					GraphActions->TryCreatingAssetNode(AssetData, Graph, FDeprecateSlateVector2D(GraphPosition), EK2NewNodeFlags::SelectNewNode);
+				}
 			}
 		}
 	}
 
-	UBlueprint* Blueprint = FBlueprintEditorUtils::FindBlueprintForGraph(Graph);
 	if ((Blueprint != nullptr) && FBlueprintEditorUtils::IsActorBased(Blueprint))
 	{
 		float XOffset = 0.0f;
@@ -5928,6 +5933,11 @@ void UEdGraphSchema_K2::DroppedAssetsOnGraph(const TArray<FAssetData>& Assets, c
 			// Make sure we have an asset type that's registered with the component list
 			if (DestinationComponentType != nullptr)
 			{
+				if (AssetReferenceFilter.IsValid() && !AssetReferenceFilter->PassesFilter(Assets[AssetIdx]))
+				{
+					continue;
+				}
+
 				const FScopedTransaction Transaction(LOCTEXT("CreateAddComponentFromAsset", "Add Component From Asset"), UEdGraphSchemaImpl::ShouldActuallyTransact());
 
 				FComponentTypeEntry ComponentType = { FString(), FString(), DestinationComponentType.GetGCPtr() };
@@ -5957,47 +5967,56 @@ void UEdGraphSchema_K2::DroppedAssetsOnPin(const TArray<FAssetData>& Assets, con
 	if ((Pin->PinType.PinCategory == PC_Object) || (Pin->PinType.PinCategory == PC_Interface))
 	{
 		UClass* PinClass = Cast<UClass>(Pin->PinType.PinSubCategoryObject.Get());
-		if(PinClass != NULL)
+		if(PinClass != nullptr)
 		{
 			// Find first asset of type of the pin
-			UObject* Asset = FAssetData::GetFirstAssetDataOfClass(Assets, PinClass).GetAsset(); //-V758
-			if(Asset != NULL)
+			FAssetData AssetData = FAssetData::GetFirstAssetDataOfClass(Assets, PinClass);
+			if (AssetData.IsValid())
 			{
-				TrySetDefaultObject(*Pin, Asset);
+				TSharedPtr<IAssetReferenceFilter> AssetReferenceFilter = MakeAssetReferenceFilter(FBlueprintEditorUtils::FindBlueprintForNode(Pin->GetOwningNode()));
+				if (!AssetReferenceFilter || AssetReferenceFilter->PassesFilter(AssetData))
+				{
+					UObject* Asset = AssetData.GetAsset(); //-V758
+					if (Asset != nullptr)
+					{
+						TrySetDefaultObject(*Pin, Asset);
+					}
+				}
 			}
 		}
 	}
 }
 
-void UEdGraphSchema_K2::GetAssetsNodeHoverMessage(const TArray<FAssetData>& Assets, const UEdGraphNode* HoverNode, FString& OutTooltipText, bool& OutOkIcon) const 
+void UEdGraphSchema_K2::GetAssetsPinHoverMessage(const TArray<FAssetData>& Assets, const UEdGraphPin* HoverPin, FString& OutTooltipText, bool& OutOkIcon) const
 { 
-	// No comment at the moment because this doesn't do anything
-	OutTooltipText = TEXT("");
-	OutOkIcon = false;
-}
-
-void UEdGraphSchema_K2::GetAssetsPinHoverMessage(const TArray<FAssetData>& Assets, const UEdGraphPin* HoverPin, FString& OutTooltipText, bool& OutOkIcon) const 
-{ 
-	OutTooltipText = TEXT("");
+	OutTooltipText.Reset();
 	OutOkIcon = false;
 
 	// If dropping onto an 'object' pin, try and set the literal
 	if ((HoverPin->PinType.PinCategory == PC_Object) || (HoverPin->PinType.PinCategory == PC_Interface))
 	{
 		UClass* PinClass = Cast<UClass>(HoverPin->PinType.PinSubCategoryObject.Get());
-		if(PinClass != NULL)
+		if(PinClass != nullptr)
 		{
 			// Find first asset of type of the pin
 			FAssetData AssetData = FAssetData::GetFirstAssetDataOfClass(Assets, PinClass);
 			if(AssetData.IsValid())
 			{
-				OutOkIcon = true;
-				OutTooltipText = FString::Printf(TEXT("Assign %s to this pin"), *(AssetData.AssetName.ToString()));
+				TSharedPtr<IAssetReferenceFilter> AssetReferenceFilter = MakeAssetReferenceFilter(FBlueprintEditorUtils::FindBlueprintForNode(HoverPin->GetOwningNode()));
+				FText FailureReason;
+				if (AssetReferenceFilter.IsValid() && !AssetReferenceFilter->PassesFilter(AssetData, &FailureReason))
+				{
+					OutTooltipText = FailureReason.ToString();
+				}
+				else
+				{
+					OutOkIcon = true;
+					OutTooltipText = FText::FormatOrdered(LOCTEXT("SupportedAssetTypeForPinDragDrop", "Assign {0} to this pin"), FText::FromName(AssetData.AssetName)).ToString();
+				}
 			}
 			else
 			{
-				OutOkIcon = false;
-				OutTooltipText = FString::Printf(TEXT("Not compatible with this pin"));
+				OutTooltipText = LOCTEXT("UnsupportedAssetTypeForPinDragDrop", "Not compatible with this pin").ToString();
 			}
 		}
 	}
@@ -6006,36 +6025,44 @@ void UEdGraphSchema_K2::GetAssetsPinHoverMessage(const TArray<FAssetData>& Asset
 
 void UEdGraphSchema_K2::GetAssetsGraphHoverMessage(const TArray<FAssetData>& Assets, const UEdGraph* HoverGraph, FString& OutTooltipText, bool& OutOkIcon) const
 {
+	OutTooltipText.Reset();
 	OutOkIcon = false;
 
 	UBlueprint* Blueprint = FBlueprintEditorUtils::FindBlueprintForGraph(HoverGraph);
 	if ((Blueprint != nullptr) && FBlueprintEditorUtils::IsActorBased(Blueprint))
 	{
-		OutTooltipText = LOCTEXT("UnsupportedAssetTypeForGraphDragDrop", "Cannot create a node from this type of asset").ToString();
+		TSharedPtr<IAssetReferenceFilter> AssetReferenceFilter = MakeAssetReferenceFilter(Blueprint);
+
 		const FBlueprintGraphModule& Module = FModuleManager::LoadModuleChecked<FBlueprintGraphModule>("BlueprintGraph");
-		bool bFoundCustomText = false;
 		for (const FAssetData& AssetData : Assets)
 		{
 			// check asset to see if we can get some AssetBlueprintGraphActions for it
 			if (const FAssetBlueprintGraphActions* GraphActions = Module.GetAssetBlueprintGraphActions(AssetData.GetClass()))
 			{
-				// get the text from the module we loaded
-				FText CustomText = GraphActions->GetGraphHoverMessage(AssetData, HoverGraph);
-				if (!CustomText.IsEmpty())
+				if (AssetReferenceFilter)
 				{
-					// want to make sure the hover message properly represents that Input Action nodes can only be dragged onto event graphs
-					if (FBlueprintEditorUtils::IsEventGraph(HoverGraph))
+					FText FailureReason;
+					if (!AssetReferenceFilter->PassesFilter(AssetData, &FailureReason))
 					{
-						OutOkIcon = true;
-						OutTooltipText = CustomText.ToString();
-					}
-					else
-					{
-						OutOkIcon = false;
-						OutTooltipText = LOCTEXT("UnsupportedAssetTypeForGraphDragDropEventGraph", "Cannot create a node from this type of asset in this graph").ToString();
+						if (OutTooltipText.IsEmpty())
+						{
+							OutTooltipText = FailureReason.ToString();
+						}
+						continue;
 					}
+				}
+
+				// want to make sure the hover message properly represents that Input Action nodes can only be dragged onto event graphs
+				if (FBlueprintEditorUtils::IsEventGraph(HoverGraph))
+				{
+					OutTooltipText = GraphActions->GetGraphHoverMessage(AssetData, HoverGraph).ToString();
+					OutOkIcon = true;
 					return;
 				}
+				else
+				{
+					OutTooltipText = LOCTEXT("UnsupportedAssetTypeForGraphDragDropEventGraph", "Cannot create a node from this type of asset in this graph").ToString();
+				}
 			}
 
 			if (UObject* Asset = AssetData.GetAsset())
@@ -6064,12 +6091,30 @@ void UEdGraphSchema_K2::GetAssetsGraphHoverMessage(const TArray<FAssetData>& Ass
 
 				if (DestinationComponentType != nullptr)
 				{
+					if (AssetReferenceFilter)
+					{
+						FText FailureReason;
+						if (!AssetReferenceFilter->PassesFilter(AssetData, &FailureReason))
+						{
+							if (OutTooltipText.IsEmpty())
+							{
+								OutTooltipText = FailureReason.ToString();
+							}
+							continue;
+						}
+					}
+
 					OutOkIcon = true;
-					OutTooltipText = TEXT("");
+					OutTooltipText.Reset();
 					return;
 				}
 			}
 		}
+
+		if (OutTooltipText.IsEmpty())
+		{
+			OutTooltipText = LOCTEXT("UnsupportedAssetTypeForGraphDragDrop", "Cannot create a node from this type of asset").ToString();
+		}
 	}
 	else
 	{
@@ -7827,6 +7872,18 @@ bool UEdGraphSchema_K2::SafeDeleteNodeFromGraph(UEdGraph* Graph, UEdGraphNode* N
 	return true;
 }
 
+TSharedPtr<IAssetReferenceFilter> UEdGraphSchema_K2::MakeAssetReferenceFilter(const UBlueprint* Blueprint)
+{
+	if (GEditor && Blueprint)
+	{
+		FAssetReferenceFilterContext AssetReferenceFilterContext;
+		AssetReferenceFilterContext.AddReferencingAsset(Blueprint);
+		return GEditor->MakeAssetReferenceFilter(AssetReferenceFilterContext);
+	}
+
+	return {};
+}
+
 #if WITH_EDITORONLY_DATA
 
 //////////////////////////////////////////////////////////////////////////
diff --git a/Engine/Source/Editor/GraphEditor/Private/SGraphPanel.cpp b/Engine/Source/Editor/GraphEditor/Private/SGraphPanel.cpp
--- a/Engine/Source/Editor/GraphEditor/Private/SGraphPanel.cpp
+++ b/Engine/Source/Editor/GraphEditor/Private/SGraphPanel.cpp
@@ -1430,18 +1430,9 @@ FReply SGraphPanel::OnDragOver( const FGeometry& MyGeometry, const FDragDropEven
 			if (AssetOp->HasAssets())
 			{
 				const TArray<FAssetData>& HoveredAssetData = AssetOp->GetAssets();
-				FText AssetReferenceFilterFailureReason;
-				if (PassesAssetReferenceFilter(HoveredAssetData, &AssetReferenceFilterFailureReason))
-				{
-					FString TooltipTextString;
-					GraphObj->GetSchema()->GetAssetsGraphHoverMessage(HoveredAssetData, GraphObj, /*out*/ TooltipTextString, /*out*/ bOkIcon);
-					TooltipText = FText::FromString(TooltipTextString);
-				}
-				else
-				{
-					TooltipText = AssetReferenceFilterFailureReason;
-					bOkIcon = false;
-				}
+				FString TooltipTextString;
+				GraphObj->GetSchema()->GetAssetsGraphHoverMessage(HoveredAssetData, GraphObj, /*out*/ TooltipTextString, /*out*/ bOkIcon);
+				TooltipText = FText::FromString(TooltipTextString);
 			}
 			const FSlateBrush* TooltipIcon = bOkIcon ? FAppStyle::GetBrush(TEXT("Graph.ConnectorFeedback.OK")) : FAppStyle::GetBrush(TEXT("Graph.ConnectorFeedback.Error"));
 			AssetOp->SetToolTip(TooltipText, TooltipIcon);
@@ -1505,10 +1496,7 @@ FReply SGraphPanel::OnDrop( const FGeometry& MyGeometry, const FDragDropEvent& D
 
 			if ( DroppedAssetData.Num() > 0 )
 			{
-				if (PassesAssetReferenceFilter(DroppedAssetData))
-				{
-					GraphObj->GetSchema()->DroppedAssetsOnGraph( DroppedAssetData, NodeAddPosition, GraphObj );
-				}
+				GraphObj->GetSchema()->DroppedAssetsOnGraph(DroppedAssetData, NodeAddPosition, GraphObj);
 				return FReply::Handled();
 			}
 		}
@@ -1517,32 +1505,6 @@ FReply SGraphPanel::OnDrop( const FGeometry& MyGeometry, const FDragDropEvent& D
 	}
 }
 
-bool SGraphPanel::PassesAssetReferenceFilter(const TArray<FAssetData>& ReferencedAssets, FText* OutFailureReason) const
-{
-	if (GEditor)
-	{
-		FAssetReferenceFilterContext AssetReferenceFilterContext;
-		UObject* GraphOuter = GraphObj ? GraphObj->GetOuter() : nullptr;
-		if (GraphOuter)
-		{
-			AssetReferenceFilterContext.AddReferencingAsset(FAssetData(GraphOuter));
-		}
-		TSharedPtr<IAssetReferenceFilter> AssetReferenceFilter = GEditor->MakeAssetReferenceFilter(AssetReferenceFilterContext);
-		if (AssetReferenceFilter.IsValid())
-		{
-			for (const FAssetData& Asset : ReferencedAssets)
-			{
-				if (!AssetReferenceFilter->PassesFilter(Asset, OutFailureReason))
-				{
-					return false;
-				}
-			}
-		}
-	}
-
-	return true;
-}
-
 TSharedPtr<SGraphNode> SGraphPanel::GetGraphNodeUnderMouse(const FGeometry& MyGeometry, const FPointerEvent& MouseEvent)
 {
 	FArrangedChildren ArrangedNodes(EVisibility::Visible);
diff --git a/Engine/Source/Editor/GraphEditor/Public/SGraphPanel.h b/Engine/Source/Editor/GraphEditor/Public/SGraphPanel.h
--- a/Engine/Source/Editor/GraphEditor/Public/SGraphPanel.h
+++ b/Engine/Source/Editor/GraphEditor/Public/SGraphPanel.h
@@ -51,7 +51,6 @@ class SWidget;
 class UEdGraph;
 class UEdGraphNode;
 class UObject;
-struct FAssetData;
 struct FDiffSingleResult;
 struct FEdGraphEditAction;
 struct FGuid;
@@ -456,9 +455,6 @@ private:
 	/** Returns the pin that we're considering as hovered if we are hovering over a spline; may be null */
 	UE_API class SGraphPin* GetBestPinFromHoveredSpline() const;
 
-	/** Returns true if all assets in the ReferencedAssets array are allowed to be referenced by this graph according to the AssetReferenceFilter. OutFailureReason if supplied if it is not.  */
-	UE_API bool PassesAssetReferenceFilter(const TArray<FAssetData>& ReferencedAssets, FText* OutFailureReason = nullptr) const;
-
 	/** Returns the top most graph node under the mouse pointer. Returns nullptr if no node found */
 	UE_API TSharedPtr<SGraphNode> GetGraphNodeUnderMouse(const FGeometry& MyGeometry, const FPointerEvent& MouseEvent);
 
diff --git a/Engine/Source/Editor/MaterialEditor/Private/MaterialEditor.cpp b/Engine/Source/Editor/MaterialEditor/Private/MaterialEditor.cpp
--- a/Engine/Source/Editor/MaterialEditor/Private/MaterialEditor.cpp
+++ b/Engine/Source/Editor/MaterialEditor/Private/MaterialEditor.cpp
@@ -47,6 +47,7 @@
 #include "Dialogs/Dialogs.h"
 #include "UnrealEdGlobals.h"
 #include "Editor.h"
+#include "Editor/EditorEngine.h"
 #include "MaterialEditorModule.h"
 #include "MaterialEditingLibrary.h"
 #include "HAL/PlatformApplicationMisc.h"
@@ -6587,6 +6588,18 @@ FText FMaterialEditor::GetOriginalObjectName() const
 	return FText::FromString(GetEditingObjects()[0]->GetName());
 }
 
+TSharedPtr<IAssetReferenceFilter> FMaterialEditor::MakeAssetReferenceFilter() const
+{
+	if (GEditor)
+	{
+		FAssetReferenceFilterContext AssetReferenceFilterContext;
+		AssetReferenceFilterContext.AddReferencingAssets(GetEditingObjects());
+		return GEditor->MakeAssetReferenceFilter(AssetReferenceFilterContext);
+	}
+
+	return {};
+}
+
 void FMaterialEditor::UpdateSubstrateTopologyPreview()
 {
 	if (Substrate::IsSubstrateEnabled())
diff --git a/Engine/Source/Editor/MaterialEditor/Private/MaterialEditor.h b/Engine/Source/Editor/MaterialEditor/Private/MaterialEditor.h
--- a/Engine/Source/Editor/MaterialEditor/Private/MaterialEditor.h
+++ b/Engine/Source/Editor/MaterialEditor/Private/MaterialEditor.h
@@ -391,6 +391,7 @@ public:
 	virtual void JumpToExpression(UMaterialExpression* Expression) override;
 	virtual void DeleteSelectedNodes() override;
 	virtual FText GetOriginalObjectName() const override;
+	virtual TSharedPtr<IAssetReferenceFilter> MakeAssetReferenceFilter() const override;
 	virtual void UpdateMaterialAfterGraphChange() override;
 	virtual void MarkMaterialDirty() override;
 	virtual void JumpToHyperlink(const UObject* ObjectReference) override; 
diff --git a/Engine/Source/Editor/MaterialEditor/Private/MaterialEditorUtilities.cpp b/Engine/Source/Editor/MaterialEditor/Private/MaterialEditorUtilities.cpp
--- a/Engine/Source/Editor/MaterialEditor/Private/MaterialEditorUtilities.cpp
+++ b/Engine/Source/Editor/MaterialEditor/Private/MaterialEditorUtilities.cpp
@@ -135,6 +135,16 @@ FText FMaterialEditorUtilities::GetOriginalObjectName(const class UEdGraph* Grap
 	return FText::GetEmpty();
 }
 
+TSharedPtr<IAssetReferenceFilter> FMaterialEditorUtilities::MakeMaterialAssetReferenceFilter(const UEdGraph* Graph)
+{
+	if (TSharedPtr<IMaterialEditor> MaterialEditor = GetIMaterialEditorForObject(Graph))
+	{
+		return MaterialEditor->MakeAssetReferenceFilter();
+	}
+
+	return {};
+}
+
 void FMaterialEditorUtilities::UpdateMaterialAfterGraphChange(const class UEdGraph* Graph)
 {
 	TSharedPtr<class IMaterialEditor> MaterialEditor = GetIMaterialEditorForObject(Graph);
diff --git a/Engine/Source/Editor/MaterialEditor/Private/MaterialInstanceEditor.cpp b/Engine/Source/Editor/MaterialEditor/Private/MaterialInstanceEditor.cpp
--- a/Engine/Source/Editor/MaterialEditor/Private/MaterialInstanceEditor.cpp
+++ b/Engine/Source/Editor/MaterialEditor/Private/MaterialInstanceEditor.cpp
@@ -46,6 +46,7 @@
 #include "Widgets/Docking/SDockTab.h"
 #include "CanvasTypes.h"
 #include "UnrealEdGlobals.h"
+#include "Editor/EditorEngine.h"
 #include "Editor/UnrealEdEngine.h"
 #include "AdvancedPreviewSceneModule.h"
 #include "ContentBrowserModule.h"
@@ -957,6 +958,18 @@ void FMaterialInstanceEditor::ExtendToolbar()
 	AddToolbarExtender(MaterialEditorModule->GetToolBarExtensibilityManager()->GetAllExtenders(GetToolkitCommands(), GetEditingObjects()));
 }
 
+TSharedPtr<IAssetReferenceFilter> FMaterialInstanceEditor::MakeAssetReferenceFilter() const
+{
+	if (GEditor)
+	{
+		FAssetReferenceFilterContext AssetReferenceFilterContext;
+		AssetReferenceFilterContext.AddReferencingAssets(GetEditingObjects());
+		return GEditor->MakeAssetReferenceFilter(AssetReferenceFilterContext);
+	}
+
+	return {};
+}
+
 void FMaterialInstanceEditor::GenerateInheritanceMenu(UToolMenu* Menu)
 {
 	RebuildInheritanceList();
diff --git a/Engine/Source/Editor/MaterialEditor/Private/MaterialInstanceEditor.h b/Engine/Source/Editor/MaterialEditor/Private/MaterialInstanceEditor.h
--- a/Engine/Source/Editor/MaterialEditor/Private/MaterialInstanceEditor.h
+++ b/Engine/Source/Editor/MaterialEditor/Private/MaterialInstanceEditor.h
@@ -120,6 +120,7 @@ public:
 	virtual void NotifyExternalMaterialChange() override;
 
 	// IMaterial Editor Interface
+	virtual TSharedPtr<IAssetReferenceFilter> MakeAssetReferenceFilter() const override;
 	virtual void GenerateInheritanceMenu(class UToolMenu* Menu) override;
 
 	void RefreshPreviewViewport();
diff --git a/Engine/Source/Editor/MaterialEditor/Public/IMaterialEditor.h b/Engine/Source/Editor/MaterialEditor/Public/IMaterialEditor.h
--- a/Engine/Source/Editor/MaterialEditor/Public/IMaterialEditor.h
+++ b/Engine/Source/Editor/MaterialEditor/Public/IMaterialEditor.h
@@ -9,6 +9,7 @@
 class FCanvas;
 class FMaterialRenderProxy;
 class FViewport;
+class IAssetReferenceFilter;
 class IMaterialEditorModule;
 class UMaterialExpression;
 class UMaterialExpressionComment;
@@ -85,6 +86,11 @@ public:
 	 */
 	virtual FText GetOriginalObjectName() const {return FText::GetEmpty();}
 
+	/**
+	 * Creates an IAssetReferenceFilter that can be used to check if an asset is valid to be used with this material.
+	 */
+	virtual TSharedPtr<IAssetReferenceFilter> MakeAssetReferenceFilter() const { return {}; }
+
 	/**
 	 * Re-links the material and updates its representation in the editor,
 	 * used when graph is changed outside of editor code.
diff --git a/Engine/Source/Editor/MaterialEditor/Public/MaterialEditorUtilities.h b/Engine/Source/Editor/MaterialEditor/Public/MaterialEditorUtilities.h
--- a/Engine/Source/Editor/MaterialEditor/Public/MaterialEditorUtilities.h
+++ b/Engine/Source/Editor/MaterialEditor/Public/MaterialEditorUtilities.h
@@ -10,6 +10,7 @@
 
 #define UE_API MATERIALEDITOR_API
 
+class IAssetReferenceFilter;
 class IMaterialEditor;
 class UMaterial;
 class UMaterialExpressionComment;
@@ -128,6 +129,13 @@ public:
 	 */
 	static UE_API FText GetOriginalObjectName(const class UEdGraph* Graph);
 
+	/**
+	 * Creates an IAssetReferenceFilter that can be used to check if an asset is valid to be used with this material.
+	 	 *
+	 * @param	Graph	Graph representing a material or material function.
+	 */
+	static UE_API TSharedPtr<IAssetReferenceFilter> MakeMaterialAssetReferenceFilter(const class UEdGraph* Graph);
+
 	/**
 	 * Re-links the material and updates its representation in the editor
 	 *
diff --git a/Engine/Source/Editor/UnrealEd/Classes/Editor/EditorEngine.h b/Engine/Source/Editor/UnrealEd/Classes/Editor/EditorEngine.h
--- a/Engine/Source/Editor/UnrealEd/Classes/Editor/EditorEngine.h
+++ b/Engine/Source/Editor/UnrealEd/Classes/Editor/EditorEngine.h
@@ -353,7 +353,9 @@ ENUM_CLASS_FLAGS(EAssetReferenceFilterRole);
 /** Struct used in filtering allowed references between assets. Passes context about the referencers to game-level filters */
 struct FAssetReferenceFilterContext
 {
+	UNREALED_API void AddReferencingAsset(const UObject* InReferencingAsset, FAssetData::ECreationFlags InCreationFlags = FAssetData::ECreationFlags::None, EAssetReferenceFilterProperties InProperties = EAssetReferenceFilterProperties::None);
 	UNREALED_API void AddReferencingAsset(const FAssetData& InReferencingAsset, EAssetReferenceFilterProperties InProperties = EAssetReferenceFilterProperties::None);
+	UNREALED_API void AddReferencingAssets(TConstArrayView<UObject*> InReferencingAssets, FAssetData::ECreationFlags InCreationFlags = FAssetData::ECreationFlags::None, EAssetReferenceFilterProperties InProperties = EAssetReferenceFilterProperties::None);
 	UNREALED_API void AddReferencingAssets(const TArray<FAssetData>& InReferencingAssets, EAssetReferenceFilterProperties InProperties = EAssetReferenceFilterProperties::None);
 	UNREALED_API void AddReferencingAssetsFromPropertyHandle(const TSharedPtr<class IPropertyHandle>& PropertyHandle, FAssetData::ECreationFlags InCreationFlags = FAssetData::ECreationFlags::None);
 
diff --git a/Engine/Source/Editor/UnrealEd/Private/EditorEngine.cpp b/Engine/Source/Editor/UnrealEd/Private/EditorEngine.cpp
--- a/Engine/Source/Editor/UnrealEd/Private/EditorEngine.cpp
+++ b/Engine/Source/Editor/UnrealEd/Private/EditorEngine.cpp
@@ -475,6 +475,20 @@ void FPreviewPlatformInfo::InternalSetFriendlyName()
 	}
 }
 
+void FAssetReferenceFilterContext::AddReferencingAsset(const UObject* InReferencingAsset, FAssetData::ECreationFlags InCreationFlags, EAssetReferenceFilterProperties InProperties)
+{
+	if (InReferencingAsset != nullptr)
+	{
+		FAssetData ReferencingObjectData(InReferencingAsset, InCreationFlags);
+
+		ReferencingAssetInfo.Emplace(ReferencingObjectData, InProperties);
+
+		PRAGMA_DISABLE_DEPRECATION_WARNINGS
+		ReferencingAssets.Add(InReferencingAsset);
+		PRAGMA_ENABLE_DEPRECATION_WARNINGS
+	}
+}
+
 void FAssetReferenceFilterContext::AddReferencingAsset(const FAssetData& InReferencingAsset, EAssetReferenceFilterProperties InProperties)
 {
 	if (InReferencingAsset.IsValid())
@@ -487,11 +501,19 @@ void FAssetReferenceFilterContext::AddReferencingAsset(const FAssetData& InRefer
 	}
 }
 
+void FAssetReferenceFilterContext::AddReferencingAssets(TConstArrayView<UObject*> InReferencingAssets, FAssetData::ECreationFlags InCreationFlags, EAssetReferenceFilterProperties InProperties)
+{
+	for (const UObject* ReferencingAsset : InReferencingAssets)
+	{
+		AddReferencingAsset(ReferencingAsset, InCreationFlags, InProperties);
+	}
+}
+
 void FAssetReferenceFilterContext::AddReferencingAssets(const TArray<FAssetData>& InReferencingAssets, EAssetReferenceFilterProperties InProperties)
 {
 	for (const FAssetData& ReferencingAsset : InReferencingAssets)
 	{
-		AddReferencingAsset(ReferencingAsset);
+		AddReferencingAsset(ReferencingAsset, InProperties);
 	}
 }
 
@@ -525,15 +547,7 @@ void FAssetReferenceFilterContext::AddReferencingAssetsFromPropertyHandle(const
 
 	for (UObject* ReferencingObject : ReferencingObjects)
 	{
-		if (ReferencingObject)
-		{
-			FAssetData ReferencingObjectData(ReferencingObject, InCreationFlags);
-			ReferencingAssetInfo.Emplace(ReferencingObjectData, ReferencerProperties);
-			
-			PRAGMA_DISABLE_DEPRECATION_WARNINGS
-			ReferencingAssets.Add(ReferencingObjectData);
-			PRAGMA_ENABLE_DEPRECATION_WARNINGS
-		}
+		AddReferencingAsset(ReferencingObject, InCreationFlags, ReferencerProperties);
 	}
 }
 
diff --git a/Engine/Source/Editor/UnrealEd/Private/MaterialGraphSchema.cpp b/Engine/Source/Editor/UnrealEd/Private/MaterialGraphSchema.cpp
--- a/Engine/Source/Editor/UnrealEd/Private/MaterialGraphSchema.cpp
+++ b/Engine/Source/Editor/UnrealEd/Private/MaterialGraphSchema.cpp
@@ -45,6 +45,7 @@
 #include "GraphEditorActions.h"
 #include "GraphEditorSettings.h"
 #include "AssetRegistry/AssetRegistryModule.h"
+#include "Editor/AssetReferenceFilter.h"
 #include "IAssetTools.h"
 #include "MaterialEditorActions.h"
 #include "MaterialGraphNode_Knot.h"
@@ -806,10 +807,16 @@ void UMaterialGraphSchema::DroppedAssetsOnGraph(const TArray<struct FAssetData>&
 
 	FVector2f ExpressionPosition = GraphPosition;
 
-	for (int32 AssetIdx = 0; AssetIdx < Assets.Num(); ++AssetIdx)
+	TSharedPtr<IAssetReferenceFilter> AssetReferenceFilter = FMaterialEditorUtilities::MakeMaterialAssetReferenceFilter(Graph);
+	for (const FAssetData& AssetData : Assets)
 	{
+		if (AssetReferenceFilter.IsValid() && !AssetReferenceFilter->PassesFilter(AssetData))
+		{
+			continue;
+		}
+
 		bool bAddedNode = false;
-		UObject* Asset = Assets[AssetIdx].GetAsset();
+		UObject* Asset = AssetData.GetAsset();
 		UClass* MaterialExpressionClass = Cast<UClass>(Asset);
 		UMaterialFunctionInterface* Func = Cast<UMaterialFunctionInterface>(Asset);
 		UTexture* Tex = Cast<UTexture>(Asset);
@@ -1178,38 +1185,41 @@ bool UMaterialGraphSchema::SafeDeleteNodeFromGraph(UEdGraph* Graph, UEdGraphNode
 
 void UMaterialGraphSchema::GetAssetsGraphHoverMessage(const TArray<FAssetData>& Assets, const UEdGraph* HoverGraph, FString& OutTooltipText, bool& OutOkIcon) const
 {
+	OutTooltipText.Reset();
 	OutOkIcon = false;
 
-	for (int32 AssetIdx = 0; AssetIdx < Assets.Num(); ++AssetIdx)
+	TSharedPtr<IAssetReferenceFilter> AssetReferenceFilter = FMaterialEditorUtilities::MakeMaterialAssetReferenceFilter(HoverGraph);
+	for (const FAssetData& AssetData : Assets)
 	{
-		UObject* Asset = Assets[AssetIdx].GetAsset();
+		UObject* Asset = AssetData.GetAsset();
 		UClass* MaterialExpressionClass = Cast<UClass>(Asset);
 		UMaterialFunctionInterface* Func = Cast<UMaterialFunctionInterface>(Asset);
 		UTexture* Tex = Cast<UTexture>(Asset);
 		USparseVolumeTexture* SparseVolumeTexture = Cast<USparseVolumeTexture>(Asset);
 		UMaterialParameterCollection* ParameterCollection = Cast<UMaterialParameterCollection>(Asset);
 
-		if (MaterialExpressionClass && MaterialExpressionClass->IsChildOf(UMaterialExpression::StaticClass()))
-		{
-			OutOkIcon = true;
-		}
-		else if (Func)
-		{
-			OutOkIcon = true;
-		}
-		else if (Tex)
-		{
-			OutOkIcon = true;
-		}
-		else if (SparseVolumeTexture)
+		if ((MaterialExpressionClass != nullptr && MaterialExpressionClass->IsChildOf(UMaterialExpression::StaticClass())) ||
+			Func != nullptr ||
+			Tex != nullptr ||
+			SparseVolumeTexture != nullptr ||
+			ParameterCollection != nullptr)
 		{
-			OutOkIcon = true;
-		}
-		else if (ParameterCollection)
-		{
-			OutOkIcon = true;
+			FText FailureReason;
+			if (AssetReferenceFilter.IsValid() && !AssetReferenceFilter->PassesFilter(AssetData, &FailureReason))
+			{
+				OutTooltipText = FailureReason.ToString();
+			}
+			else
+			{
+				OutOkIcon = true;
+			}
 		}
 	}
+
+	if (OutOkIcon)
+	{
+		OutTooltipText.Reset();
+	}
 }
 
 #if WITH_EDITORONLY_DATA
