[UE] Attempt to add more context to reported cases where we are failing to resolve a node's owning Blueprint and/or graph, and avoid crashing the editor as a result.

diff --git a/Engine/Source/Editor/AnimGraph/Private/K2Node_AnimGetter.cpp b/Engine/Source/Editor/AnimGraph/Private/K2Node_AnimGetter.cpp
--- a/Engine/Source/Editor/AnimGraph/Private/K2Node_AnimGetter.cpp
+++ b/Engine/Source/Editor/AnimGraph/Private/K2Node_AnimGetter.cpp
@@ -579,6 +579,12 @@ FText UK2Node_AnimGetter::GenerateTitle(UFunction* Getter, UAnimStateNodeBase* S
 
 void UK2Node_AnimGetter::UpdateCachedTitle()
 {
+	// Make sure we have a valid Blueprint owner before proceeding. Otherwise we will assert trying to resolve the target function below.
+	if (!ensureMsgf(HasValidBlueprint(), TEXT("UK2Node_AnimGetter::UpdateCachedTitle(): Attempting to update cached title for [%s] without a valid Blueprint owner (this is unexpected)."), *GetFullName()))
+	{
+		return;
+	}
+
 	if (UFunction* Function = GetTargetFunction())
 	{
 		CachedTitle = GenerateTitle(Function, SourceStateNode, SourceNode);
diff --git a/Engine/Source/Editor/BlueprintGraph/Private/EdGraphSchema_K2.cpp b/Engine/Source/Editor/BlueprintGraph/Private/EdGraphSchema_K2.cpp
--- a/Engine/Source/Editor/BlueprintGraph/Private/EdGraphSchema_K2.cpp
+++ b/Engine/Source/Editor/BlueprintGraph/Private/EdGraphSchema_K2.cpp
@@ -7778,11 +7778,25 @@ void UEdGraphSchema_K2::HandleParameterDefaultValueChanged(UK2Node* InTargetNode
 			return;
 		}
 
+		// Make sure we have a valid Blueprint owner before proceeding, so we don't crash the editor attempting to modify it below.
+		UBlueprint* TargetBlueprint = FBlueprintEditorUtils::FindBlueprintForNode(EditablePinNode);
+		if (!ensureMsgf(TargetBlueprint, TEXT("UEdGraphSchema_K2::HandleParameterDefaultValueChanged(): Failed to find a Blueprint for target node: %s"), *GetPathNameSafe(EditablePinNode)))
+		{
+			return;
+		}
+
+		// Make sure we have a valid outer graph object before proceeding, so we don't crash the editor attempting to modify it below.
+		UEdGraph* EditablePinNodeGraph = EditablePinNode->GetGraph();
+		if (!ensureMsgf(EditablePinNodeGraph, TEXT("UEdGraphSchema_K2::HandleParameterDefaultValueChanged(): Failed to find a valid graph for target node: %s"), *GetPathNameSafe(EditablePinNode)))
+		{
+			return;
+		}
+
 		FParamsChangedHelper ParamsChangedHelper;
-		ParamsChangedHelper.ModifiedBlueprints.Add(FBlueprintEditorUtils::FindBlueprintForNode(InTargetNode));
-		FBlueprintEditorUtils::MarkBlueprintAsStructurallyModified(FBlueprintEditorUtils::FindBlueprintForNode(InTargetNode));
+		ParamsChangedHelper.ModifiedBlueprints.Add(TargetBlueprint);
+		FBlueprintEditorUtils::MarkBlueprintAsStructurallyModified(TargetBlueprint);
 
-		ParamsChangedHelper.Broadcast(FBlueprintEditorUtils::FindBlueprintForNode(InTargetNode), EditablePinNode, InTargetNode->GetGraph());
+		ParamsChangedHelper.Broadcast(TargetBlueprint, EditablePinNode, EditablePinNodeGraph);
 
 		for (UEdGraph* ModifiedGraph : ParamsChangedHelper.ModifiedGraphs)
 		{
diff --git a/Engine/Source/Editor/BlueprintGraph/Private/K2Node.cpp b/Engine/Source/Editor/BlueprintGraph/Private/K2Node.cpp
--- a/Engine/Source/Editor/BlueprintGraph/Private/K2Node.cpp
+++ b/Engine/Source/Editor/BlueprintGraph/Private/K2Node.cpp
@@ -708,6 +708,12 @@ void UK2Node::ReconstructNode()
 	// Clear previously set messages
 	ErrorMsg.Reset();
 
+	// Make sure we have a valid Blueprint owner, so we don't crash the editor attempting to access it below.
+	if (!ensureMsgf(HasValidBlueprint(), TEXT("UK2Node::ReconstructNode(): Attempting to reconstruct [%s] without a valid Blueprint owner (this is unexpected)."), *GetFullName()))
+	{
+		return;
+	}
+
 	UBlueprint* Blueprint = GetBlueprint();
 
 	FLinkerLoad* Linker = Blueprint->GetLinker();
