BlueprintNodeTemplateCache::IsTemplateOuter now handles ParentGraph being nullptr

diff --git a/Engine/Source/Editor/BlueprintGraph/Private/BlueprintNodeTemplateCache.cpp b/Engine/Source/Editor/BlueprintGraph/Private/BlueprintNodeTemplateCache.cpp
--- a/Engine/Source/Editor/BlueprintGraph/Private/BlueprintNodeTemplateCache.cpp
+++ b/Engine/Source/Editor/BlueprintGraph/Private/BlueprintNodeTemplateCache.cpp
@@ -225,7 +225,7 @@ static void BlueprintNodeTemplateCacheImpl::MarkGraphForTemplateUse(UEdGraph* Ne
 //------------------------------------------------------------------------------
 bool BlueprintNodeTemplateCacheImpl::IsTemplateOuter(UEdGraph* ParentGraph)
 {
-	if (ParentGraph->HasAnyFlags(RF_Transactional))
+	if (ParentGraph && ParentGraph->HasAnyFlags(RF_Transactional))
 	{
 		UPackage* GraphPackage = ParentGraph->GetPackage();
 		FMetaData& PackageMetadata = GraphPackage->GetMetaData();
