Merge 46415098 for hotfix
[Blueprint]
- Fix crash during force delete of some blueprint classes used in a scene

[RN] Fixed crash after force delete of some blueprints

diff --git a/Engine/Source/Editor/BlueprintGraph/Private/K2Node_FunctionEntry.cpp b/Engine/Source/Editor/BlueprintGraph/Private/K2Node_FunctionEntry.cpp
--- a/Engine/Source/Editor/BlueprintGraph/Private/K2Node_FunctionEntry.cpp
+++ b/Engine/Source/Editor/BlueprintGraph/Private/K2Node_FunctionEntry.cpp
@@ -352,9 +352,13 @@ void UK2Node_FunctionEntry::AddReferencedObjects(UObject* InThis, FReferenceColl
 {
 	UK2Node_FunctionEntry* This = CastChecked<UK2Node_FunctionEntry>(InThis);
 	TSharedPtr<FStructOnScope> FunctionVariableCache = This->FunctionVariableCache;
-	if(FunctionVariableCache.IsValid())
+	if (FunctionVariableCache.IsValid())
 	{
-		Collector.AddPropertyReferences(FunctionVariableCache->GetStruct(), FunctionVariableCache->GetStructMemory());
+		TWeakObjectPtr<const UStruct> StructPtr = FunctionVariableCache->GetStructPtr();
+		if (StructPtr.IsValid())
+		{
+			Collector.AddPropertyReferences(StructPtr.Get(), FunctionVariableCache->GetStructMemory());
+		}
 	}
 	return UK2Node_FunctionTerminator::AddReferencedObjects(InThis, Collector);
 }
