Validate property name for function result pins

[RN] Output pins on blueprint result nodes are once again validated for uniqueness when being edited

diff --git a/Engine/Source/Editor/BlueprintGraph/Classes/K2Node_FunctionResult.h b/Engine/Source/Editor/BlueprintGraph/Classes/K2Node_FunctionResult.h
--- a/Engine/Source/Editor/BlueprintGraph/Classes/K2Node_FunctionResult.h
+++ b/Engine/Source/Editor/BlueprintGraph/Classes/K2Node_FunctionResult.h
@@ -46,6 +46,7 @@ class UK2Node_FunctionResult : public UK2Node_FunctionTerminator
 	virtual void PostPasteNode() override;
 	virtual bool CanUserDeleteNode() const override;
 	virtual void FixupPinStringDataReferences(FArchive* SavingArchive) override;
+	virtual ERenamePinResult RenameUserDefinedPinImpl(FName OldName, FName NewName, bool bTest) override;
 	//~ End UK2Node Interface
 
 	//~ Begin UK2Node_EditablePinBase Interface
diff --git a/Engine/Source/Editor/BlueprintGraph/Private/K2Node_FunctionResult.cpp b/Engine/Source/Editor/BlueprintGraph/Private/K2Node_FunctionResult.cpp
--- a/Engine/Source/Editor/BlueprintGraph/Private/K2Node_FunctionResult.cpp
+++ b/Engine/Source/Editor/BlueprintGraph/Private/K2Node_FunctionResult.cpp
@@ -208,6 +208,29 @@ void UK2Node_FunctionResult::FixupPinStringDataReferences(FArchive* SavingArchiv
 	}
 }
 
+ERenamePinResult UK2Node_FunctionResult::RenameUserDefinedPinImpl(FName OldName, FName NewName, bool bTest)
+{
+	ERenamePinResult Result = Super::RenameUserDefinedPinImpl(OldName, NewName, bTest);
+
+	const bool bCheckLocalVariablesNameCollision =
+		(Result == ERenamePinResult::ERenamePinResult_Success) &&
+		bTest;
+	
+	if (bCheckLocalVariablesNameCollision)
+	{
+		if (UFunction* const Function = FindSignatureFunction())
+		{
+			const FProperty* ExistingProperty = FindFProperty<const FProperty>(Function, NewName);
+			if (ExistingProperty)
+			{
+				Result = ERenamePinResult_NameCollision;
+			}
+		}
+	}
+
+	return Result;
+}
+
 FNodeHandlingFunctor* UK2Node_FunctionResult::CreateNodeHandler(FKismetCompilerContext& CompilerContext) const
 {
 	return new FKCHandler_FunctionResult(CompilerContext);
