Fix non determinism in some text literal values (localization guid regenerating on compile)

[RN] Fixed rare non-determinism in blueprint generated bytecode caused by text values that matched their default string
#Synced-CL 48774400

diff --git a/Engine/Source/Editor/BlueprintGraph/Private/EdGraphSchema_K2.cpp b/Engine/Source/Editor/BlueprintGraph/Private/EdGraphSchema_K2.cpp
--- a/Engine/Source/Editor/BlueprintGraph/Private/EdGraphSchema_K2.cpp
+++ b/Engine/Source/Editor/BlueprintGraph/Private/EdGraphSchema_K2.cpp
@@ -5321,15 +5321,18 @@ static void ConformAutogeneratedDefaultValuePackage(
 			const TCHAR* Success = FTextStringHelper::ReadFromBuffer(*AutogeneratedDefaultValue, AutogeneratedDefaultTextValue);
 			check(Success);
 
-			// Conform the auto-generated default against this package ID, preserving its key to avoid determinism issues
-			const FText ConformedAutogeneratedDefaultTextValue = TextNamespaceUtil::CopyTextToPackage(AutogeneratedDefaultTextValue, PackageNamespace, TextNamespaceUtil::ETextCopyMethod::PreserveKey);
-
-			// IdenticalTo is a quick test for whether CopyTextToPackage returned the same text it was given (meaning there's nothing to update)
-			if (!ConformedAutogeneratedDefaultTextValue.IdenticalTo(AutogeneratedDefaultTextValue))
+			// Does auto-generated default text have a mismatched package ID?
+			if (!AutogeneratedDefaultTextValue.IsEmpty() && !AutogeneratedDefaultTextValue.IsCultureInvariant() && !AutogeneratedDefaultTextValue.IsFromStringTable())
 			{
-				// Fix-up the auto-generated default from the conformed value
-				AutogeneratedDefaultValue.Reset();
-				FTextStringHelper::WriteToBuffer(AutogeneratedDefaultValue, ConformedAutogeneratedDefaultTextValue);
+				if (const FString CurrentPackageNamespace = TextNamespaceUtil::ExtractPackageNamespace(FTextInspector::GetNamespace(AutogeneratedDefaultTextValue).Get(FString()));
+					!CurrentPackageNamespace.IsEmpty() && !CurrentPackageNamespace.Equals(PackageNamespace, ESearchCase::CaseSensitive))
+				{
+					// Conform the auto-generated default against this package ID, preserving its key to avoid determinism issues
+					AutogeneratedDefaultTextValue = TextNamespaceUtil::CopyTextToPackage(AutogeneratedDefaultTextValue, PackageNamespace, TextNamespaceUtil::ETextCopyMethod::PreserveKey);
+					// Fix-up the auto-generated default from the conformed value
+					AutogeneratedDefaultValue.Reset();
+					FTextStringHelper::WriteToBuffer(AutogeneratedDefaultValue, AutogeneratedDefaultTextValue);
+				}
 			}
 		}
 	}
