PR #14072: Copy PureState for UK2Node_DynamicCast

Any dynamic cast node copied in blueprints will now include PureState, respecting whether or not the cast had been converted to a pure cast.

Previously any cast nodes converted to pure would be pasted as impure nodes. The behavior now matches function calls, with UK2Node_CallFunction also copying node purity. "NodePurityOverride," in that case.

#Synced-CL 48842190

diff --git a/Engine/Source/Editor/BlueprintGraph/Classes/K2Node_DynamicCast.h b/Engine/Source/Editor/BlueprintGraph/Classes/K2Node_DynamicCast.h
--- a/Engine/Source/Editor/BlueprintGraph/Classes/K2Node_DynamicCast.h
+++ b/Engine/Source/Editor/BlueprintGraph/Classes/K2Node_DynamicCast.h
@@ -25,6 +25,14 @@ class FString;
 class UEdGraphPin;
 class UObject;
 
+UENUM()
+enum class EPureState : uint8
+{
+	Pure,
+	Impure,
+	UseDefault
+};
+
 UCLASS(MinimalAPI)
 class UK2Node_DynamicCast : public UK2Node
 {
@@ -110,14 +118,9 @@ protected:
 	bool bIsPureCast;
 
 private:
-	enum class EPureState : uint8
-	{
-		Pure,
-		Impure,
-		UseDefault
-	};
 
 	/** Internal state used to determine purity */
+	UPROPERTY()
 	EPureState PureState;
 
 	/** [DEPRECATED] Receives the old value on load if previously saved */
