Optimize bp context menu by avoiding creation of node templates for use as node spawners, move principal functionality to a private nodespawner thus improving encapsulation

[RN] The blueprint context menu no longer creates function node templates for most functions, reducing object graph memory
#Synced-CL 49062257

diff --git a/Engine/Source/Editor/BlueprintGraph/Private/BlueprintActionDatabase.cpp b/Engine/Source/Editor/BlueprintGraph/Private/BlueprintActionDatabase.cpp
--- a/Engine/Source/Editor/BlueprintGraph/Private/BlueprintActionDatabase.cpp
+++ b/Engine/Source/Editor/BlueprintGraph/Private/BlueprintActionDatabase.cpp
@@ -10,6 +10,7 @@
 #include "BlueprintAssetHandler.h"
 #include "BlueprintBoundEventNodeSpawner.h"
 #include "BlueprintComponentNodeSpawner.h"
+#include "BlueprintGraph/BlueprintPrivateFunctionNodeSpawner.h"
 #include "BlueprintDelegateNodeSpawner.h"
 #include "BlueprintEditorSettings.h"
 #include "BlueprintEventNodeSpawner.h"
@@ -699,7 +700,12 @@ static void BlueprintActionDatabaseImpl::AddClassFunctionActions(UClass const* c
 			//        Blueprint "implemented interface", then we don't need to 
 			//        include it (the function is accounted for in from the 
 			//        interface class).
-			UBlueprintFunctionNodeSpawner* FuncSpawner = UBlueprintFunctionNodeSpawner::Create(Function);
+
+			UBlueprintNodeSpawner* FuncSpawner = UBlueprintPrivateFunctionNodeSpawner::Create(Function);
+			if(!FuncSpawner)
+			{
+				FuncSpawner =  UBlueprintFunctionNodeSpawner::Create(Function);
+			}
 			
 			// Only add this action to the list of the operator function is not already registered. Otherwise we will 
 			// get a bunch of duplicate operator actions
diff --git a/Engine/Source/Editor/BlueprintGraph/Private/BlueprintFunctionNodeSpawner.cpp b/Engine/Source/Editor/BlueprintGraph/Private/BlueprintFunctionNodeSpawner.cpp
--- a/Engine/Source/Editor/BlueprintGraph/Private/BlueprintFunctionNodeSpawner.cpp
+++ b/Engine/Source/Editor/BlueprintGraph/Private/BlueprintFunctionNodeSpawner.cpp
@@ -4,6 +4,7 @@
 
 #include "BlueprintActionFilter.h"
 #include "BlueprintEditorSettings.h"
+#include "BlueprintGraph/BlueprintFunctionNodeSpawnerImpl.h"
 #include "BlueprintNodeSpawnerUtils.h"
 #include "BlueprintNodeTemplateCache.h"
 #include "BlueprintTypePromotion.h"
@@ -51,156 +52,6 @@
 
 #define LOCTEXT_NAMESPACE "BlueprintFunctionNodeSpawner"
 
-/*******************************************************************************
- * Static UBlueprintFunctionNodeSpawner Helpers
- ******************************************************************************/
-
-//------------------------------------------------------------------------------
-namespace BlueprintFunctionNodeSpawnerImpl
-{
-	FVector2D BindingOffset = FVector2D::ZeroVector;
-	static const FText FallbackCategory = LOCTEXT("UncategorizedFallbackCategory", "Call Function");
-
-	/**
-	 * 
-	 * 
-	 * @param  NewNode	
-	 * @param  BoundObject	
-	 * @return 
-	 */
-	static bool BindFunctionNode(UK2Node_CallFunction* NewNode, FBindingObject BoundObject);
-
-	/**
-	 * 
-	 * 
-	 * @param  NewNode	
-	 * @param  BindingSpawner	
-	 * @return 
-	 */
-	template <class NodeType>
-	static bool BindFunctionNode(UK2Node_CallFunction* NewNode, UBlueprintNodeSpawner* BindingSpawner);
-
-	/**
-	 * 
-	 * 
-	 * @param  InputNode
-	 * @return 
-	 */
-	static FVector2D CalculateBindingPosition(UEdGraphNode* const InputNode);
-
-	/**
-	 * 
-	 * 
-	 * @param  Struct	
-	 * @param  Function	
-	 * @param  OperatorMetaTag	
-	 * @return 
-	 */
-	static bool IsStructOperatorFunc(const UScriptStruct* Struct, const UFunction* Function, FName const OperatorMetaTag);
-}
-
-//------------------------------------------------------------------------------
-static bool BlueprintFunctionNodeSpawnerImpl::BindFunctionNode(UK2Node_CallFunction* NewNode, FBindingObject BoundObject)
-{
-	bool bSuccessfulBinding = false;
-
-	bool const bIsTemplateNode = FBlueprintNodeTemplateCache::IsTemplateOuter(NewNode->GetGraph());
-	if (!bIsTemplateNode)
-	{
-		if (FProperty const* BoundProperty = BoundObject.Get<FProperty>())
-		{
-			if (UK2Node_CallFunctionOnMember* CallOnMemberNode = Cast<UK2Node_CallFunctionOnMember>(NewNode))
-			{
-				// force bIsConsideredSelfContext to false, else the target 
-				// could end up being the skeleton class (and functionally, 
-				// there is no difference)
-				CallOnMemberNode->MemberVariableToCallOn.SetFromField<FProperty>(BoundProperty, /*bIsConsideredSelfContext =*/false);
-				bSuccessfulBinding = true;
-				CallOnMemberNode->ReconstructNode();
-			}
-			else
-			{
-				UBlueprintNodeSpawner* TempNodeSpawner = UBlueprintVariableNodeSpawner::CreateFromMemberOrParam(UK2Node_VariableGet::StaticClass(), BoundProperty);
-				bSuccessfulBinding = BindFunctionNode<UK2Node_VariableGet>(NewNode, TempNodeSpawner);
-			}
-		}
-		else if (AActor* BoundActor = BoundObject.Get<AActor>())
-		{
-			auto PostSpawnSetupLambda = [](UEdGraphNode* InNewNode, bool /*bIsTemplateNode*/, AActor* ActorInst)
-			{
-				UK2Node_Literal* ActorRefNode = CastChecked<UK2Node_Literal>(InNewNode);
-				ActorRefNode->SetObjectRef(ActorInst);
-			};
-			UBlueprintNodeSpawner::FCustomizeNodeDelegate PostSpawnDelegate = UBlueprintNodeSpawner::FCustomizeNodeDelegate::CreateStatic(PostSpawnSetupLambda, BoundActor);
-			
-			UBlueprintNodeSpawner* TempNodeSpawner = UBlueprintNodeSpawner::Create<UK2Node_Literal>(/*Outer =*/GetTransientPackage(), PostSpawnDelegate);
-			bSuccessfulBinding = BindFunctionNode<UK2Node_Literal>(NewNode, TempNodeSpawner);
-		}		
-	}
-	return bSuccessfulBinding;
-}
-
-//------------------------------------------------------------------------------
-template <class NodeType>
-static bool BlueprintFunctionNodeSpawnerImpl::BindFunctionNode(UK2Node_CallFunction* NewNode, UBlueprintNodeSpawner* BindingSpawner)
-{
-	bool bSuccessfulBinding = false;
-
-	FVector2D BindingPos = CalculateBindingPosition(NewNode);
-	UEdGraph* ParentGraph = NewNode->GetGraph();
-	NodeType* BindingNode = CastChecked<NodeType>(BindingSpawner->Invoke(ParentGraph, IBlueprintNodeBinder::FBindingSet(), BindingPos));
-
-	BindingOffset.Y += UEdGraphSchema_K2::EstimateNodeHeight(BindingNode);
-
-	UEdGraphPin* LiteralOutput = BindingNode->GetValuePin();
-	UEdGraphPin* CallSelfInput = NewNode->FindPin(UEdGraphSchema_K2::PN_Self);
-	// connect the new "get-var" node with the spawned function node
-	if ((LiteralOutput != nullptr) && (CallSelfInput != nullptr))
-	{
-		LiteralOutput->MakeLinkTo(CallSelfInput);
-		bSuccessfulBinding = true;
-	}
-
-	return bSuccessfulBinding;
-}
-
-//------------------------------------------------------------------------------
-static FVector2D BlueprintFunctionNodeSpawnerImpl::CalculateBindingPosition(UEdGraphNode* const InputNode)
-{
-	float const EstimatedVarNodeWidth = 224.0f;
-	FVector2D AttachingNodePos;
-	AttachingNodePos.X = InputNode->NodePosX - EstimatedVarNodeWidth;
-
-	float const EstimatedVarNodeHeight = 48.0f;
-	float const EstimatedFuncNodeHeight = UEdGraphSchema_K2::EstimateNodeHeight(InputNode);
-	float const FuncNodeMidYCoordinate = InputNode->NodePosY + (EstimatedFuncNodeHeight / 2.0f);
-	AttachingNodePos.Y = FuncNodeMidYCoordinate - (EstimatedVarNodeWidth / 2.0f);
-
-	AttachingNodePos += BindingOffset;
-	return AttachingNodePos;
-}
-
-//------------------------------------------------------------------------------
-static bool BlueprintFunctionNodeSpawnerImpl::IsStructOperatorFunc(const UScriptStruct* Struct, const UFunction* Function, FName const OperatorMetaTag)
-{
-	bool bIsOperatorFunc = false;
-
-	FString NamedOperatorFunction = Struct->GetMetaData(OperatorMetaTag);
-	if (!NamedOperatorFunction.IsEmpty())
-	{
-		UObject* OperatorOuter = nullptr;
-		if (ResolveName(OperatorOuter, NamedOperatorFunction, /*Create =*/false, /*Throw =*/false))
-		{
-			if ((Function->GetOuter() == OperatorOuter) &&
-				(Function->GetName()  == NamedOperatorFunction))
-			{
-				bIsOperatorFunc = true;
-			}
-		}
-	}
-	return bIsOperatorFunc;
-}
-
 /*******************************************************************************
  * UBlueprintFunctionNodeSpawner
  ******************************************************************************/
diff --git a/Engine/Source/Editor/BlueprintGraph/Private/BlueprintGraph/BlueprintFunctionNodeSpawnerImpl.cpp b/Engine/Source/Editor/BlueprintGraph/Private/BlueprintGraph/BlueprintFunctionNodeSpawnerImpl.cpp
new file mode 100644
--- /dev/null
+++ b/Engine/Source/Editor/BlueprintGraph/Private/BlueprintGraph/BlueprintFunctionNodeSpawnerImpl.cpp
@@ -0,0 +1,400 @@
+// Copyright Epic Games, Inc. All Rights Reserved.
+
+#include "BlueprintGraph/BlueprintFunctionNodeSpawnerImpl.h"
+
+#include "BlueprintVariableNodeSpawner.h"
+#include "BlueprintTypePromotion.h"
+#include "BlueprintNodeTemplateCache.h"
+#include "EditorCategoryUtils.h"
+#include "EdGraphSchema_K2.h"
+#include "GameFramework/Actor.h"
+#include "K2Node_CallFunction.h"
+#include "K2Node_CallFunctionOnMember.h"
+#include "K2Node_CallArrayFunction.h"
+#include "K2Node_CallDataTableFunction.h"
+#include "K2Node_CallMaterialParameterCollectionFunction.h"
+#include "K2Node_CommutativeAssociativeBinaryOperator.h"
+#include "K2Node_Event.h"
+#include "K2Node_Literal.h"
+#include "K2Node_PromotableOperator.h"
+#include "K2Node_VariableGet.h"
+#include "Kismet2/BlueprintEditorUtils.h"
+#include "ObjectTools.h"
+
+#define LOCTEXT_NAMESPACE "BlueprintFunctionNodeSpawner"
+ 
+FVector2D BlueprintFunctionNodeSpawnerImpl::BindingOffset = FVector2D::ZeroVector;
+const FText BlueprintFunctionNodeSpawnerImpl::FallbackCategory = LOCTEXT("UncategorizedFallbackCategory", "Call Function");
+
+//------------------------------------------------------------------------------
+bool BlueprintFunctionNodeSpawnerImpl::BindFunctionNode(UK2Node_CallFunction* NewNode, FBindingObject BoundObject)
+{
+	bool bSuccessfulBinding = false;
+
+	bool const bIsTemplateNode = FBlueprintNodeTemplateCache::IsTemplateOuter(NewNode->GetGraph());
+	if (!bIsTemplateNode)
+	{
+		if (FProperty const* BoundProperty = BoundObject.Get<FProperty>())
+		{
+			if (UK2Node_CallFunctionOnMember* CallOnMemberNode = Cast<UK2Node_CallFunctionOnMember>(NewNode))
+			{
+				// force bIsConsideredSelfContext to false, else the target 
+				// could end up being the skeleton class (and functionally, 
+				// there is no difference)
+				CallOnMemberNode->MemberVariableToCallOn.SetFromField<FProperty>(BoundProperty, /*bIsConsideredSelfContext =*/false);
+				bSuccessfulBinding = true;
+				CallOnMemberNode->ReconstructNode();
+			}
+			else
+			{
+				UBlueprintNodeSpawner* TempNodeSpawner = UBlueprintVariableNodeSpawner::CreateFromMemberOrParam(UK2Node_VariableGet::StaticClass(), BoundProperty);
+				bSuccessfulBinding = BindFunctionNode<UK2Node_VariableGet>(NewNode, TempNodeSpawner);
+			}
+		}
+		else if (AActor* BoundActor = BoundObject.Get<AActor>())
+		{
+			auto PostSpawnSetupLambda = [](UEdGraphNode* InNewNode, bool /*bIsTemplateNode*/, AActor* ActorInst)
+			{
+				UK2Node_Literal* ActorRefNode = CastChecked<UK2Node_Literal>(InNewNode);
+				ActorRefNode->SetObjectRef(ActorInst);
+			};
+			UBlueprintNodeSpawner::FCustomizeNodeDelegate PostSpawnDelegate = UBlueprintNodeSpawner::FCustomizeNodeDelegate::CreateStatic(PostSpawnSetupLambda, BoundActor);
+			
+			UBlueprintNodeSpawner* TempNodeSpawner = UBlueprintNodeSpawner::Create<UK2Node_Literal>(/*Outer =*/GetTransientPackage(), PostSpawnDelegate);
+			bSuccessfulBinding = BindFunctionNode<UK2Node_Literal>(NewNode, TempNodeSpawner);
+		}		
+	}
+	return bSuccessfulBinding;
+}
+
+//------------------------------------------------------------------------------
+template <class NodeType>
+bool BlueprintFunctionNodeSpawnerImpl::BindFunctionNode(UK2Node_CallFunction* NewNode, UBlueprintNodeSpawner* BindingSpawner)
+{
+	bool bSuccessfulBinding = false;
+
+	FVector2D BindingPos = CalculateBindingPosition(NewNode);
+	UEdGraph* ParentGraph = NewNode->GetGraph();
+	NodeType* BindingNode = CastChecked<NodeType>(BindingSpawner->Invoke(ParentGraph, IBlueprintNodeBinder::FBindingSet(), BindingPos));
+
+	BindingOffset.Y += UEdGraphSchema_K2::EstimateNodeHeight(BindingNode);
+
+	UEdGraphPin* LiteralOutput = BindingNode->GetValuePin();
+	UEdGraphPin* CallSelfInput = NewNode->FindPin(UEdGraphSchema_K2::PN_Self);
+	// connect the new "get-var" node with the spawned function node
+	if ((LiteralOutput != nullptr) && (CallSelfInput != nullptr))
+	{
+		LiteralOutput->MakeLinkTo(CallSelfInput);
+		bSuccessfulBinding = true;
+	}
+
+	return bSuccessfulBinding;
+}
+
+//------------------------------------------------------------------------------
+FVector2D BlueprintFunctionNodeSpawnerImpl::CalculateBindingPosition(UEdGraphNode* const InputNode)
+{
+	float const EstimatedVarNodeWidth = 224.0f;
+	FVector2D AttachingNodePos;
+	AttachingNodePos.X = InputNode->NodePosX - EstimatedVarNodeWidth;
+
+	float const EstimatedVarNodeHeight = 48.0f;
+	float const EstimatedFuncNodeHeight = UEdGraphSchema_K2::EstimateNodeHeight(InputNode);
+	float const FuncNodeMidYCoordinate = InputNode->NodePosY + (EstimatedFuncNodeHeight / 2.0f);
+	AttachingNodePos.Y = FuncNodeMidYCoordinate - (EstimatedVarNodeWidth / 2.0f);
+
+	AttachingNodePos += BindingOffset;
+	return AttachingNodePos;
+}
+
+//------------------------------------------------------------------------------
+bool BlueprintFunctionNodeSpawnerImpl::IsStructOperatorFunc(const UScriptStruct* Struct, const UFunction* Function, FName const OperatorMetaTag)
+{
+	bool bIsOperatorFunc = false;
+
+	FString NamedOperatorFunction = Struct->GetMetaData(OperatorMetaTag);
+	if (!NamedOperatorFunction.IsEmpty())
+	{
+		UObject* OperatorOuter = nullptr;
+		if (ResolveName(OperatorOuter, NamedOperatorFunction, /*Create =*/false, /*Throw =*/false))
+		{
+			if ((Function->GetOuter() == OperatorOuter) &&
+				(Function->GetName()  == NamedOperatorFunction))
+			{
+				bIsOperatorFunc = true;
+			}
+		}
+	}
+	return bIsOperatorFunc;
+}
+
+#undef LOCTEXT_NAMESPACE
+
+// This code was moved out of K2Node_CallFunction, matching localization key used there:
+
+#define LOCTEXT_NAMESPACE "K2Node"
+
+namespace UE::Private
+{
+	bool CanPasteFunctionHere(const UFunction* Function, const UEdGraph* TargetGraph);
+}
+
+bool UE::Private::CanPasteFunctionHere(const UFunction* Function, const UEdGraph* TargetGraph)
+{
+	// Basic check for graph compatibility, etc.
+	const UEdGraphSchema* Schema = TargetGraph->GetSchema();
+	bool bCanPaste = Schema->GetClass()->IsChildOf(UEdGraphSchema_K2::StaticClass());
+
+	// Cannot paste editor only functions into runtime graphs
+	if (bCanPaste)
+	{
+		bCanPaste = UK2Node_CallFunction::CanEditorOnlyFunctionBeCalled(Function, TargetGraph);
+	}
+
+	// We check function context for placability only in the base class case; derived classes are typically bound to
+	// specific functions that should always be placeable, but may not always be explicitly callable (e.g. InternalUseOnly).
+	if (bCanPaste)
+	{
+		const UEdGraphSchema_K2* K2Schema = GetDefault<UEdGraphSchema_K2>();
+		uint32 AllowedFunctionTypes = UEdGraphSchema_K2::EFunctionType::FT_Pure | 
+			UEdGraphSchema_K2::EFunctionType::FT_Const | 
+			UEdGraphSchema_K2::EFunctionType::FT_Protected;
+		if (K2Schema->DoesGraphSupportImpureFunctions(TargetGraph))
+		{
+			AllowedFunctionTypes |= UEdGraphSchema_K2::EFunctionType::FT_Imperative;
+		}
+
+		if (Function)
+		{
+			bCanPaste = K2Schema->CanFunctionBeUsedInGraph(
+				FBlueprintEditorUtils::FindBlueprintForGraphChecked(TargetGraph)->SkeletonGeneratedClass, 
+				Function, 
+				TargetGraph, 
+				AllowedFunctionTypes, 
+				false);
+		}
+	}
+	
+	return bCanPaste;
+}
+
+FText BlueprintFunctionNodeSpawnerImpl::GetNodeTitle(const UFunction* Function, ENodeTitleType::Type TitleType)
+{
+	check(Function);
+	FText FunctionName;
+	FText ContextString;
+	FText RPCString;
+
+	// matching UBlueprintFunctionNodeSpawner::Create
+	// but there may be tension with UK2Node_CallFunction::GetNodeTitle:
+	constexpr bool bAllowFriendlyNames = true; 
+	
+	RPCString = UK2Node_Event::GetLocalizedNetString(Function->FunctionFlags, true);
+	FunctionName = ObjectTools::GetUserFacingFunctionName(Function, bAllowFriendlyNames);
+	ContextString = GetFunctionContextString(Function);
+
+	if(TitleType == ENodeTitleType::FullTitle)
+	{
+		FFormatNamedArguments Args;
+		Args.Add(TEXT("FunctionName"), FunctionName);
+		Args.Add(TEXT("ContextString"), ContextString);
+		Args.Add(TEXT("RPCString"), RPCString);
+
+		if (ContextString.IsEmpty() && RPCString.IsEmpty())
+		{
+			return FText::Format(LOCTEXT("CallFunction_FullTitle", "{FunctionName}"), Args);
+		}
+		else if (ContextString.IsEmpty())
+		{
+			return FText::Format(LOCTEXT("CallFunction_FullTitle_WithRPCString", "{FunctionName}\n{RPCString}"), Args);
+		}
+		else if (RPCString.IsEmpty())
+		{
+			return FText::Format(LOCTEXT("CallFunction_FullTitle_WithContextString", "{FunctionName}\n{ContextString}"), Args);
+		}
+		else
+		{
+			return FText::Format(LOCTEXT("CallFunction_FullTitle_WithContextRPCString", "{FunctionName}\n{ContextString}\n{RPCString}"), Args);
+		}
+	}
+	else
+	{
+		return FunctionName;
+	}
+}
+
+FText BlueprintFunctionNodeSpawnerImpl::GetMenuCategory(const UFunction* Function)
+{
+	FText Category = UK2Node_CallFunction::GetDefaultCategoryForFunction(Function, FText::GetEmpty());
+	if (Category.IsEmpty())
+	{
+		UClass* FunctionClass = Function->GetOwnerClass()->GetAuthoritativeClass();
+		Category = FEditorCategoryUtils::BuildCategoryString( 
+			FCommonEditorCategory::Class, 
+			FunctionClass->GetDisplayNameText() );
+	}
+	return Category;
+}
+
+bool BlueprintFunctionNodeSpawnerImpl::IsActionFilteredOut(const UFunction* TargetFunction, FBlueprintActionFilter const& Filter)
+{
+	bool bIsFilteredOut = false;
+	for(UEdGraph* TargetGraph : Filter.Context.Graphs)
+	{
+		bIsFilteredOut |= !UE::Private::CanPasteFunctionHere(TargetFunction, TargetGraph);
+	}
+
+	const bool bIsProtected = (TargetFunction->FunctionFlags & FUNC_Protected) != 0;
+	const bool bIsPrivate = (TargetFunction->FunctionFlags & FUNC_Private) != 0;
+	const UClass* OwningClass = TargetFunction->GetOwnerClass();
+	if( (bIsProtected || bIsPrivate) && !FBlueprintEditorUtils::IsNativeSignature(TargetFunction) && OwningClass)
+	{
+		OwningClass = OwningClass->GetAuthoritativeClass();
+		// we can filter private and protected blueprints that are unrelated:
+		bool bAccessibleInAll = true;
+		for (const UBlueprint* Blueprint : Filter.Context.Blueprints)
+		{
+			UClass* AuthoritativeClass = Blueprint->GeneratedClass;
+			if(!AuthoritativeClass)
+			{
+				continue;
+			}
+
+			if(bIsPrivate)
+			{
+				bAccessibleInAll = bAccessibleInAll && AuthoritativeClass == OwningClass;
+			}
+			else if(bIsProtected)
+			{
+				bAccessibleInAll = bAccessibleInAll && AuthoritativeClass->IsChildOf(OwningClass);
+			}
+		}
+
+		if(!bAccessibleInAll)
+		{
+			bIsFilteredOut = true;
+		}
+	}
+
+	return bIsFilteredOut;
+}
+
+FText BlueprintFunctionNodeSpawnerImpl::GetTooltipText(const UFunction* Function)
+{
+	FText BaseTooltip = FText::FromString(ObjectTools::GetDefaultTooltipForFunction(Function));
+
+	FFormatNamedArguments Args;
+	Args.Add(TEXT("DefaultTooltip"), BaseTooltip);
+
+	if (Function->HasAllFunctionFlags(FUNC_BlueprintAuthorityOnly))
+	{
+		Args.Add(
+			TEXT("Subtitle"),
+			NSLOCTEXT("K2Node", "ServerFunction", "Authority Only. This function will only execute on the server.")
+		);
+		// FText::Format() is slow, so we cache this to save on performance
+		return FText::Format(LOCTEXT("CallFunction_SubtitledTooltip", "{DefaultTooltip}\n\n{Subtitle}"), Args);
+	}
+	else if (Function->HasAllFunctionFlags(FUNC_BlueprintCosmetic))
+	{
+		Args.Add(
+			TEXT("Subtitle"),
+			NSLOCTEXT("K2Node", "ClientFunction", "Cosmetic. This event is only for cosmetic, non-gameplay actions.")
+		);
+		return FText::Format(LOCTEXT("CallFunction_SubtitledTooltip", "{DefaultTooltip}\n\n{Subtitle}"), Args);
+	} 
+	else if (Function->HasMetaData(FBlueprintMetadata::MD_Latent))
+	{
+		Args.Add(
+			TEXT("Subtitle"),
+			NSLOCTEXT("K2Node", "LatentFunction", "Latent. This node will complete at a later time. Latent nodes can only be placed in event graphs.")
+		);
+		return FText::Format(LOCTEXT("CallFunction_SubtitledTooltip", "{DefaultTooltip}\n\n{Subtitle}"), Args);
+	}
+	else
+	{
+		return BaseTooltip;
+	}
+}
+
+FSlateIcon BlueprintFunctionNodeSpawnerImpl::GetIconAndTint(const UFunction* Function, FLinearColor& OutColor)
+{
+	return UK2Node_CallFunction::GetPaletteIconForFunction(Function, OutColor);
+}
+
+FString BlueprintFunctionNodeSpawnerImpl::GetDocumentationLink(const UFunction* Function)
+{
+	UClass* ParentClass = Function->GetOwnerClass();
+	if (ParentClass != NULL)
+	{
+		return FString::Printf(TEXT("Shared/GraphNodes/Blueprint/%s%s"), ParentClass->GetPrefixCPP(), *ParentClass->GetName());
+	}
+
+	return FString("Shared/GraphNodes/Blueprint/UK2Node_CallFunction");
+}
+
+FString BlueprintFunctionNodeSpawnerImpl::GetDocumentationExcerptName(const UFunction* Function)
+{
+	return Function->GetName();
+}
+
+FBlueprintNodeSignature BlueprintFunctionNodeSpawnerImpl::GetSignature(const UFunction* Function)
+{
+	return FBlueprintNodeSignature(UK2Node_CallFunction::StaticClass());
+}
+
+FText BlueprintFunctionNodeSpawnerImpl::GetFunctionContextString(const UFunction* Function)
+{
+	check(Function);
+	UClass* CurrentSelfClass = (Function != NULL) ? Function->GetOwnerClass() : NULL;
+	UClass const* TrueSelfClass = CurrentSelfClass;
+	if (CurrentSelfClass && CurrentSelfClass->ClassGeneratedBy)
+	{
+		TrueSelfClass = CurrentSelfClass->GetAuthoritativeClass();
+	}
+
+	const FText TargetText = FBlueprintEditorUtils::GetFriendlyClassDisplayName(TrueSelfClass);
+
+	FFormatNamedArguments Args;
+	Args.Add(TEXT("TargetName"), TargetText);
+	return FText::Format(LOCTEXT("CallFunctionOnDifferentContext", "Target is {TargetName}"), Args);
+}
+
+TSubclassOf<UK2Node_CallFunction> BlueprintFunctionNodeSpawnerImpl::GetCallFunctionNodeClass(const UFunction* Function)
+{
+	bool const bIsPure = Function->HasAllFunctionFlags(FUNC_BlueprintPure);
+	bool const bHasArrayPointerParms = Function->HasMetaData(FBlueprintMetadata::MD_ArrayParam);
+	bool const bIsCommutativeAssociativeBinaryOp = Function->HasMetaData(FBlueprintMetadata::MD_CommutativeAssociativeBinaryOperator);
+	bool const bIsMaterialParamCollectionFunc = Function->HasMetaData(FBlueprintMetadata::MD_MaterialParameterCollectionFunction);
+	bool const bIsDataTableFunc = Function->HasMetaData(FBlueprintMetadata::MD_DataTablePin);
+
+	bool const bIsPromotableFunction = TypePromoDebug::IsTypePromoEnabled() && FTypePromotion::IsFunctionPromotionReady(Function);
+	bool bIsNoTemplateRequiredLocal = false;
+
+	if (bIsPromotableFunction)
+	{
+		return UK2Node_PromotableOperator::StaticClass();
+	}
+	else if(bIsCommutativeAssociativeBinaryOp && bIsPure)
+	{
+		return UK2Node_CommutativeAssociativeBinaryOperator::StaticClass();
+	}
+	else if (bIsMaterialParamCollectionFunc)
+	{
+		return UK2Node_CallMaterialParameterCollectionFunction::StaticClass();
+	}
+	else if (bIsDataTableFunc)
+	{
+		return UK2Node_CallDataTableFunction::StaticClass();
+	}
+	else if (bHasArrayPointerParms)
+	{
+		return UK2Node_CallArrayFunction::StaticClass();
+	}
+	else
+	{
+		return UK2Node_CallFunction::StaticClass();
+	}
+}
+
+#undef LOCTEXT_NAMESPACE
diff --git a/Engine/Source/Editor/BlueprintGraph/Private/BlueprintGraph/BlueprintFunctionNodeSpawnerImpl.h b/Engine/Source/Editor/BlueprintGraph/Private/BlueprintGraph/BlueprintFunctionNodeSpawnerImpl.h
new file mode 100644
--- /dev/null
+++ b/Engine/Source/Editor/BlueprintGraph/Private/BlueprintGraph/BlueprintFunctionNodeSpawnerImpl.h
@@ -0,0 +1,46 @@
+// Copyright Epic Games, Inc. All Rights Reserved.
+
+#pragma once
+
+#include "BlueprintNodeBinder.h"
+#include "BlueprintNodeSignature.h"
+#include "Containers/UnrealString.h"
+#include "Internationalization/Text.h"
+#include "Math/Vector2D.h"
+#include "Textures/SlateIcon.h"
+
+namespace ENodeTitleType { enum Type : int; }
+struct FLinearColor;
+class FBlueprintActionFilter;
+class UBlueprintNodeSpawner;
+class UFunction;
+class UK2Node_CallFunction;
+
+/*******************************************************************************
+ * Static UBlueprintFunctionNodeSpawner Helpers
+ ******************************************************************************/
+
+//------------------------------------------------------------------------------
+namespace BlueprintFunctionNodeSpawnerImpl
+{
+	extern FVector2D BindingOffset;
+	extern const FText FallbackCategory;
+
+	bool BindFunctionNode(UK2Node_CallFunction* NewNode, FBindingObject BoundObject);
+	template <class NodeType>
+	bool BindFunctionNode(UK2Node_CallFunction* NewNode, UBlueprintNodeSpawner* BindingSpawner);
+
+	FVector2D CalculateBindingPosition(UEdGraphNode* const InputNode);
+	bool IsStructOperatorFunc(const UScriptStruct* Struct, const UFunction* Function, FName const OperatorMetaTag);
+
+	FText GetNodeTitle(const UFunction* Function, ENodeTitleType::Type TitleType);
+	FText GetMenuCategory(const UFunction* Function);
+	bool IsActionFilteredOut(const UFunction* Function, FBlueprintActionFilter const& Filter);
+	FText GetTooltipText(const UFunction* Function);
+	FSlateIcon GetIconAndTint(const UFunction* Function, FLinearColor& OutColor);
+	FString GetDocumentationLink(const UFunction* Function);
+	FString GetDocumentationExcerptName(const UFunction* Function);
+	FBlueprintNodeSignature GetSignature(const UFunction* Function);
+	FText GetFunctionContextString(const UFunction* Function);
+	TSubclassOf<UK2Node_CallFunction> GetCallFunctionNodeClass(const UFunction* Function);
+}
\ No newline at end of file
diff --git a/Engine/Source/Editor/BlueprintGraph/Private/BlueprintGraph/BlueprintPrivateFunctionNodeSpawner.cpp b/Engine/Source/Editor/BlueprintGraph/Private/BlueprintGraph/BlueprintPrivateFunctionNodeSpawner.cpp
new file mode 100644
--- /dev/null
+++ b/Engine/Source/Editor/BlueprintGraph/Private/BlueprintGraph/BlueprintPrivateFunctionNodeSpawner.cpp
@@ -0,0 +1,138 @@
+// Copyright Epic Games, Inc. All Rights Reserved.
+
+#pragma once
+
+#include "BlueprintGraph/BlueprintPrivateFunctionNodeSpawner.h"
+#include "BlueprintGraph/BlueprintFunctionNodeSpawnerImpl.h"
+#include "BlueprintNodeSpawnerUtils.h"
+#include "BlueprintTypePromotion.h"
+#include "K2Node_CallFunctionOnMember.h"
+#include "ObjectEditorUtils.h"
+#include "SNodePanel.h"
+
+UBlueprintPrivateFunctionNodeSpawner::UBlueprintPrivateFunctionNodeSpawner(const FObjectInitializer& ObjectInitializer)
+	: UBlueprintFieldNodeSpawner(ObjectInitializer)
+{
+
+}
+
+UBlueprintPrivateFunctionNodeSpawner* UBlueprintPrivateFunctionNodeSpawner::Create(const UFunction* Function)
+{
+	check(Function != nullptr);
+	
+	bool const bIsPromotableFunction =
+		TypePromoDebug::IsTypePromoEnabled() &&
+		FTypePromotion::IsFunctionPromotionReady(Function);
+	if(bIsPromotableFunction)
+	{
+		return nullptr;
+	}
+
+	// deconstruction of UBlueprintNodeSpawner::PrimeDefaultUiSpec and UBlueprintFunctionNodeSpawner::Create:
+	UBlueprintPrivateFunctionNodeSpawner* Result = NewObject<UBlueprintPrivateFunctionNodeSpawner>();
+	Result->SetField(const_cast<UFunction*>(Function));
+	Result->NodeClass = BlueprintFunctionNodeSpawnerImpl::GetCallFunctionNodeClass(Function);
+	FBlueprintActionUiSpec& MenuSignature = Result->DefaultMenuSignature;
+	MenuSignature.MenuName = BlueprintFunctionNodeSpawnerImpl::GetNodeTitle(Function, ENodeTitleType::MenuTitle);
+	MenuSignature.Category = BlueprintFunctionNodeSpawnerImpl::GetMenuCategory(Function);
+	MenuSignature.Tooltip = BlueprintFunctionNodeSpawnerImpl::GetTooltipText(Function);
+	MenuSignature.Keywords = UK2Node_CallFunction::GetKeywordsForFunction(Function);
+	MenuSignature.Icon = BlueprintFunctionNodeSpawnerImpl::GetIconAndTint(Function, MenuSignature.IconTint);
+	MenuSignature.DocLink = BlueprintFunctionNodeSpawnerImpl::GetDocumentationLink(Function);
+	MenuSignature.DocExcerptTag = BlueprintFunctionNodeSpawnerImpl::GetDocumentationExcerptName(Function);
+	return Result;
+}
+	
+void UBlueprintPrivateFunctionNodeSpawner::Prime()
+{
+	// intentionally override and nop
+}
+
+FBlueprintActionUiSpec UBlueprintPrivateFunctionNodeSpawner::GetUiSpec(FBlueprintActionContext const& Context, FBindingSet const& Bindings) const
+{
+	FBlueprintActionUiSpec MenuSignature = DefaultMenuSignature;
+	if (Bindings.Num() == 1)
+	{
+		FObjectProperty const* ObjectProperty = Bindings.CreateConstIterator()->Get<FObjectProperty>();
+		if (ObjectProperty != nullptr)
+		{
+			FString BoundFunctionName = FString::Printf(TEXT("%s (%s)"), *MenuSignature.MenuName.ToString(), *ObjectProperty->GetName());
+			MenuSignature.MenuName = FText::FromString(BoundFunctionName);
+		}
+	}
+	return MenuSignature;
+}
+
+UEdGraphNode* UBlueprintPrivateFunctionNodeSpawner::Invoke(UEdGraph* ParentGraph, FBindingSet const& Bindings, FVector2D const Location) const
+{
+	auto PostSpawnSetupLambda = [](UEdGraphNode* NewNode, bool bIsTemplateNode, UFunction const* Function, FSetNodeFieldDelegate SetFieldDelegate, FCustomizeNodeDelegate UserDelegate)
+	{
+		SetFieldDelegate.ExecuteIfBound(NewNode, const_cast<UFunction*>(Function));
+		UserDelegate.ExecuteIfBound(NewNode, bIsTemplateNode);
+	};
+	FCustomizeNodeDelegate PostSpawnSetupDelegate = FCustomizeNodeDelegate::CreateStatic(PostSpawnSetupLambda, GetFunction(), SetNodeFieldDelegate, CustomizeNodeDelegate);
+
+	UClass* SpawnClass = NodeClass;
+
+	const UBlueprintEditorSettings* BPSettings = GetDefault<UBlueprintEditorSettings>();
+
+	bool const bSpawnCallOnMember = (Bindings.Num() == 1) && Bindings.CreateConstIterator()->Get<FObjectProperty>();
+	if (bSpawnCallOnMember && (BPSettings->bCompactCallOnMemberNodes))
+	{
+		SpawnClass = UK2Node_CallFunctionOnMember::StaticClass();
+	}
+
+	// if this spawner was set up to spawn a bound node, reset this so the 
+	// bound nodes get positioned properly..
+	BlueprintFunctionNodeSpawnerImpl::BindingOffset = FVector2D::ZeroVector;
+
+	UEdGraphNode* SpawnedNode = Super::SpawnNode<UEdGraphNode>(SpawnClass, ParentGraph, Bindings, Location, PostSpawnSetupDelegate);
+	SpawnedNode->SnapToGrid(SNodePanel::GetSnapGridSize());
+	return SpawnedNode;
+}
+
+bool UBlueprintPrivateFunctionNodeSpawner::IsTemplateNodeFilteredOut(FBlueprintActionFilter const& Filter) const
+{
+	return BlueprintFunctionNodeSpawnerImpl::IsActionFilteredOut(GetFunction(), Filter);
+}
+
+bool UBlueprintPrivateFunctionNodeSpawner::IsBindingCompatible(FBindingObject BindingCandidate) const
+{
+	UFunction const* Function = GetFunction();
+	checkSlow(Function != nullptr);
+
+	if ( !ensureMsgf(!FBlueprintNodeSpawnerUtils::IsStaleFieldAction(this), 
+			TEXT("Invalid BlueprintFunctionNodeSpawner (for %s). Was the action database properly updated when this class was compiled?"), 
+			*Function->GetOwnerClass()->GetName()) )
+	{
+		return false;
+	}
+
+	bool const bNodeTypeMatches = (NodeClass == UK2Node_CallFunction::StaticClass());
+	bool bClassOwnerMatches = false;
+
+	UClass* BindingClass = FBlueprintNodeSpawnerUtils::GetBindingClass(BindingCandidate)->GetAuthoritativeClass();
+	if (UClass const* FuncOwner = Function->GetOwnerClass()->GetAuthoritativeClass())
+	{
+		bClassOwnerMatches = BindingClass && BindingClass->IsChildOf(FuncOwner);
+	}
+
+	return bNodeTypeMatches && bClassOwnerMatches && !FObjectEditorUtils::IsFunctionHiddenFromClass(Function, BindingClass);
+}
+
+bool UBlueprintPrivateFunctionNodeSpawner::CanBindMultipleObjects() const
+{
+	UFunction const* Function = GetFunction();
+	check(Function != nullptr);
+	return UK2Node_CallFunction::CanFunctionSupportMultipleTargets(Function);
+}
+
+bool UBlueprintPrivateFunctionNodeSpawner::BindToNode(UEdGraphNode* Node, FBindingObject Binding) const
+{
+	return BlueprintFunctionNodeSpawnerImpl::BindFunctionNode(CastChecked<UK2Node_CallFunction>(Node), Binding);
+}
+
+UFunction const* UBlueprintPrivateFunctionNodeSpawner::GetFunction() const
+{
+	return Cast<UFunction>(GetField().ToUObject());
+}
diff --git a/Engine/Source/Editor/BlueprintGraph/Private/BlueprintGraph/BlueprintPrivateFunctionNodeSpawner.h b/Engine/Source/Editor/BlueprintGraph/Private/BlueprintGraph/BlueprintPrivateFunctionNodeSpawner.h
new file mode 100644
--- /dev/null
+++ b/Engine/Source/Editor/BlueprintGraph/Private/BlueprintGraph/BlueprintPrivateFunctionNodeSpawner.h
@@ -0,0 +1,51 @@
+// Copyright Epic Games, Inc. All Rights Reserved.
+
+#pragma once
+
+#include "BlueprintActionFilter.h"
+#include "BlueprintFieldNodeSpawner.h"
+#include "BlueprintNodeBinder.h"
+#include "BlueprintNodeSpawner.h"
+#include "CoreMinimal.h"
+#include "Math/Vector2D.h"
+#include "Templates/SubclassOf.h"
+#include "UObject/Class.h"
+#include "UObject/Object.h"
+#include "UObject/ObjectMacros.h"
+#include "UObject/UObjectGlobals.h"
+
+#include "BlueprintPrivateFunctionNodeSpawner.generated.h"
+
+class UEdGraph;
+class UEdGraphNode;
+class UFunction;
+class UK2Node_CallFunction;
+class UObject;
+struct FBlueprintActionContext;
+
+UCLASS(Transient)
+class UBlueprintPrivateFunctionNodeSpawner final : public UBlueprintFieldNodeSpawner
+{
+	GENERATED_BODY()
+
+public:
+	static UBlueprintPrivateFunctionNodeSpawner* Create(const UFunction* Function);
+	
+private:
+	UBlueprintPrivateFunctionNodeSpawner(const FObjectInitializer& ObjectInitializer);
+	// UBlueprintNodeSpawner interface
+	void Prime() override;
+	FBlueprintActionUiSpec GetUiSpec(FBlueprintActionContext const& Context, FBindingSet const& Bindings) const override;
+	UEdGraphNode* Invoke(UEdGraph* ParentGraph, FBindingSet const& Bindings, FVector2D const Location) const override;
+	bool IsTemplateNodeFilteredOut(FBlueprintActionFilter const& Filter) const override;
+	bool DoesRequireTemplateNode() const override { return false; }
+	// End UBlueprintNodeSpawner interface
+
+	// IBlueprintNodeBinder interface
+	bool IsBindingCompatible(FBindingObject BindingCandidate) const override;
+	bool CanBindMultipleObjects() const override;
+	bool BindToNode(UEdGraphNode* Node, FBindingObject Binding) const override;
+	// End IBlueprintNodeBinder interface
+
+	UFunction const* GetFunction() const;
+};
diff --git a/Engine/Source/Editor/BlueprintGraph/Private/BlueprintNodeSpawner.cpp b/Engine/Source/Editor/BlueprintGraph/Private/BlueprintNodeSpawner.cpp
--- a/Engine/Source/Editor/BlueprintGraph/Private/BlueprintNodeSpawner.cpp
+++ b/Engine/Source/Editor/BlueprintGraph/Private/BlueprintNodeSpawner.cpp
@@ -290,6 +290,11 @@ UEdGraphNode* UBlueprintNodeSpawner::GetTemplateNode(UEdGraph* TargetGraph, FBin
 {       
 	TRACE_CPUPROFILER_EVENT_SCOPE(UBlueprintNodeSpawner::GetTemplateNode);
 	
+	if(!DoesRequireTemplateNode())
+	{
+		return nullptr;
+	}
+
 	UEdGraphNode* TemplateNode = BlueprintNodeSpawnerImpl::GetSharedTemplateCache()->GetNodeTemplate(this, TargetGraph);
 
 	if (TemplateNode && Bindings.Num() > 0) 
diff --git a/Engine/Source/Editor/BlueprintGraph/Public/BlueprintNodeSpawner.h b/Engine/Source/Editor/BlueprintGraph/Public/BlueprintNodeSpawner.h
--- a/Engine/Source/Editor/BlueprintGraph/Public/BlueprintNodeSpawner.h
+++ b/Engine/Source/Editor/BlueprintGraph/Public/BlueprintNodeSpawner.h
@@ -295,6 +295,8 @@ private:
 	 * @return Null if it failed to spawn a node (if NodeClass is null), otherwise a newly spawned node.
 	 */
 	UE_API UEdGraphNode* SpawnEdGraphNode(TSubclassOf<UEdGraphNode> NodeClass, UEdGraph* ParentGraph, const FBindingSet& Bindings, FVector2D Location, FCustomizeNodeDelegate PostSpawnDelegate) const;
+
+	virtual bool DoesRequireTemplateNode() const { return true; }
 };
 
 /*******************************************************************************
