[Backout] - CL49146491

This caused an issue with some blueprint nodes having a title of "None" and being unusable

#rnx
Original CL Desc
-----------------------------------------------------------------
Optimize bp context menu by avoiding creation of node templates for use as node spawners, move principal functionality to a private nodespawner thus improving encapsulation

[RN] The blueprint context menu no longer creates function node templates for most functions, reducing object graph memory
#Synced-CL 49062257

#Synced-CL 49127187

diff --git a/Engine/Source/Editor/BlueprintGraph/Private/BlueprintActionDatabase.cpp b/Engine/Source/Editor/BlueprintGraph/Private/BlueprintActionDatabase.cpp
--- a/Engine/Source/Editor/BlueprintGraph/Private/BlueprintActionDatabase.cpp
+++ b/Engine/Source/Editor/BlueprintGraph/Private/BlueprintActionDatabase.cpp
@@ -10,7 +10,6 @@
 #include "BlueprintAssetHandler.h"
 #include "BlueprintBoundEventNodeSpawner.h"
 #include "BlueprintComponentNodeSpawner.h"
-#include "BlueprintGraph/BlueprintPrivateFunctionNodeSpawner.h"
 #include "BlueprintDelegateNodeSpawner.h"
 #include "BlueprintEditorSettings.h"
 #include "BlueprintEventNodeSpawner.h"
@@ -700,12 +699,7 @@ static void BlueprintActionDatabaseImpl::AddClassFunctionActions(UClass const* c
 			//        Blueprint "implemented interface", then we don't need to 
 			//        include it (the function is accounted for in from the 
 			//        interface class).
-
-			UBlueprintNodeSpawner* FuncSpawner = UBlueprintPrivateFunctionNodeSpawner::Create(Function);
-			if(!FuncSpawner)
-			{
-				FuncSpawner =  UBlueprintFunctionNodeSpawner::Create(Function);
-			}
+			UBlueprintFunctionNodeSpawner* FuncSpawner = UBlueprintFunctionNodeSpawner::Create(Function);
 			
 			// Only add this action to the list of the operator function is not already registered. Otherwise we will 
 			// get a bunch of duplicate operator actions
diff --git a/Engine/Source/Editor/BlueprintGraph/Private/BlueprintFunctionNodeSpawner.cpp b/Engine/Source/Editor/BlueprintGraph/Private/BlueprintFunctionNodeSpawner.cpp
--- a/Engine/Source/Editor/BlueprintGraph/Private/BlueprintFunctionNodeSpawner.cpp
+++ b/Engine/Source/Editor/BlueprintGraph/Private/BlueprintFunctionNodeSpawner.cpp
@@ -4,7 +4,6 @@
 
 #include "BlueprintActionFilter.h"
 #include "BlueprintEditorSettings.h"
-#include "BlueprintGraph/BlueprintFunctionNodeSpawnerImpl.h"
 #include "BlueprintNodeSpawnerUtils.h"
 #include "BlueprintNodeTemplateCache.h"
 #include "BlueprintTypePromotion.h"
@@ -52,6 +51,156 @@
 
 #define LOCTEXT_NAMESPACE "BlueprintFunctionNodeSpawner"
 
+/*******************************************************************************
+ * Static UBlueprintFunctionNodeSpawner Helpers
+ ******************************************************************************/
+
+//------------------------------------------------------------------------------
+namespace BlueprintFunctionNodeSpawnerImpl
+{
+	FVector2D BindingOffset = FVector2D::ZeroVector;
+	static const FText FallbackCategory = LOCTEXT("UncategorizedFallbackCategory", "Call Function");
+
+	/**
+	 * 
+	 * 
+	 * @param  NewNode	
+	 * @param  BoundObject	
+	 * @return 
+	 */
+	static bool BindFunctionNode(UK2Node_CallFunction* NewNode, FBindingObject BoundObject);
+
+	/**
+	 * 
+	 * 
+	 * @param  NewNode	
+	 * @param  BindingSpawner	
+	 * @return 
+	 */
+	template <class NodeType>
+	static bool BindFunctionNode(UK2Node_CallFunction* NewNode, UBlueprintNodeSpawner* BindingSpawner);
+
+	/**
+	 * 
+	 * 
+	 * @param  InputNode
+	 * @return 
+	 */
+	static FVector2D CalculateBindingPosition(UEdGraphNode* const InputNode);
+
+	/**
+	 * 
+	 * 
+	 * @param  Struct	
+	 * @param  Function	
+	 * @param  OperatorMetaTag	
+	 * @return 
+	 */
+	static bool IsStructOperatorFunc(const UScriptStruct* Struct, const UFunction* Function, FName const OperatorMetaTag);
+}
+
+//------------------------------------------------------------------------------
+static bool BlueprintFunctionNodeSpawnerImpl::BindFunctionNode(UK2Node_CallFunction* NewNode, FBindingObject BoundObject)
+{
+	bool bSuccessfulBinding = false;
+
+	bool const bIsTemplateNode = FBlueprintNodeTemplateCache::IsTemplateOuter(NewNode->GetGraph());
+	if (!bIsTemplateNode)
+	{
+		if (FProperty const* BoundProperty = BoundObject.Get<FProperty>())
+		{
+			if (UK2Node_CallFunctionOnMember* CallOnMemberNode = Cast<UK2Node_CallFunctionOnMember>(NewNode))
+			{
+				// force bIsConsideredSelfContext to false, else the target 
+				// could end up being the skeleton class (and functionally, 
+				// there is no difference)
+				CallOnMemberNode->MemberVariableToCallOn.SetFromField<FProperty>(BoundProperty, /*bIsConsideredSelfContext =*/false);
+				bSuccessfulBinding = true;
+				CallOnMemberNode->ReconstructNode();
+			}
+			else
+			{
+				UBlueprintNodeSpawner* TempNodeSpawner = UBlueprintVariableNodeSpawner::CreateFromMemberOrParam(UK2Node_VariableGet::StaticClass(), BoundProperty);
+				bSuccessfulBinding = BindFunctionNode<UK2Node_VariableGet>(NewNode, TempNodeSpawner);
+			}
+		}
+		else if (AActor* BoundActor = BoundObject.Get<AActor>())
+		{
+			auto PostSpawnSetupLambda = [](UEdGraphNode* InNewNode, bool /*bIsTemplateNode*/, AActor* ActorInst)
+			{
+				UK2Node_Literal* ActorRefNode = CastChecked<UK2Node_Literal>(InNewNode);
+				ActorRefNode->SetObjectRef(ActorInst);
+			};
+			UBlueprintNodeSpawner::FCustomizeNodeDelegate PostSpawnDelegate = UBlueprintNodeSpawner::FCustomizeNodeDelegate::CreateStatic(PostSpawnSetupLambda, BoundActor);
+			
+			UBlueprintNodeSpawner* TempNodeSpawner = UBlueprintNodeSpawner::Create<UK2Node_Literal>(/*Outer =*/GetTransientPackage(), PostSpawnDelegate);
+			bSuccessfulBinding = BindFunctionNode<UK2Node_Literal>(NewNode, TempNodeSpawner);
+		}		
+	}
+	return bSuccessfulBinding;
+}
+
+//------------------------------------------------------------------------------
+template <class NodeType>
+static bool BlueprintFunctionNodeSpawnerImpl::BindFunctionNode(UK2Node_CallFunction* NewNode, UBlueprintNodeSpawner* BindingSpawner)
+{
+	bool bSuccessfulBinding = false;
+
+	FVector2D BindingPos = CalculateBindingPosition(NewNode);
+	UEdGraph* ParentGraph = NewNode->GetGraph();
+	NodeType* BindingNode = CastChecked<NodeType>(BindingSpawner->Invoke(ParentGraph, IBlueprintNodeBinder::FBindingSet(), BindingPos));
+
+	BindingOffset.Y += UEdGraphSchema_K2::EstimateNodeHeight(BindingNode);
+
+	UEdGraphPin* LiteralOutput = BindingNode->GetValuePin();
+	UEdGraphPin* CallSelfInput = NewNode->FindPin(UEdGraphSchema_K2::PN_Self);
+	// connect the new "get-var" node with the spawned function node
+	if ((LiteralOutput != nullptr) && (CallSelfInput != nullptr))
+	{
+		LiteralOutput->MakeLinkTo(CallSelfInput);
+		bSuccessfulBinding = true;
+	}
+
+	return bSuccessfulBinding;
+}
+
+//------------------------------------------------------------------------------
+static FVector2D BlueprintFunctionNodeSpawnerImpl::CalculateBindingPosition(UEdGraphNode* const InputNode)
+{
+	float const EstimatedVarNodeWidth = 224.0f;
+	FVector2D AttachingNodePos;
+	AttachingNodePos.X = InputNode->NodePosX - EstimatedVarNodeWidth;
+
+	float const EstimatedVarNodeHeight = 48.0f;
+	float const EstimatedFuncNodeHeight = UEdGraphSchema_K2::EstimateNodeHeight(InputNode);
+	float const FuncNodeMidYCoordinate = InputNode->NodePosY + (EstimatedFuncNodeHeight / 2.0f);
+	AttachingNodePos.Y = FuncNodeMidYCoordinate - (EstimatedVarNodeWidth / 2.0f);
+
+	AttachingNodePos += BindingOffset;
+	return AttachingNodePos;
+}
+
+//------------------------------------------------------------------------------
+static bool BlueprintFunctionNodeSpawnerImpl::IsStructOperatorFunc(const UScriptStruct* Struct, const UFunction* Function, FName const OperatorMetaTag)
+{
+	bool bIsOperatorFunc = false;
+
+	FString NamedOperatorFunction = Struct->GetMetaData(OperatorMetaTag);
+	if (!NamedOperatorFunction.IsEmpty())
+	{
+		UObject* OperatorOuter = nullptr;
+		if (ResolveName(OperatorOuter, NamedOperatorFunction, /*Create =*/false, /*Throw =*/false))
+		{
+			if ((Function->GetOuter() == OperatorOuter) &&
+				(Function->GetName()  == NamedOperatorFunction))
+			{
+				bIsOperatorFunc = true;
+			}
+		}
+	}
+	return bIsOperatorFunc;
+}
+
 /*******************************************************************************
  * UBlueprintFunctionNodeSpawner
  ******************************************************************************/
diff --git a/Engine/Source/Editor/BlueprintGraph/Private/BlueprintGraph/BlueprintFunctionNodeSpawnerImpl.cpp b/Engine/Source/Editor/BlueprintGraph/Private/BlueprintGraph/BlueprintFunctionNodeSpawnerImpl.cpp
deleted file mode 100644
diff --git a/Engine/Source/Editor/BlueprintGraph/Private/BlueprintGraph/BlueprintFunctionNodeSpawnerImpl.h b/Engine/Source/Editor/BlueprintGraph/Private/BlueprintGraph/BlueprintFunctionNodeSpawnerImpl.h
deleted file mode 100644
diff --git a/Engine/Source/Editor/BlueprintGraph/Private/BlueprintGraph/BlueprintPrivateFunctionNodeSpawner.cpp b/Engine/Source/Editor/BlueprintGraph/Private/BlueprintGraph/BlueprintPrivateFunctionNodeSpawner.cpp
deleted file mode 100644
diff --git a/Engine/Source/Editor/BlueprintGraph/Private/BlueprintGraph/BlueprintPrivateFunctionNodeSpawner.h b/Engine/Source/Editor/BlueprintGraph/Private/BlueprintGraph/BlueprintPrivateFunctionNodeSpawner.h
deleted file mode 100644
diff --git a/Engine/Source/Editor/BlueprintGraph/Private/BlueprintNodeSpawner.cpp b/Engine/Source/Editor/BlueprintGraph/Private/BlueprintNodeSpawner.cpp
--- a/Engine/Source/Editor/BlueprintGraph/Private/BlueprintNodeSpawner.cpp
+++ b/Engine/Source/Editor/BlueprintGraph/Private/BlueprintNodeSpawner.cpp
@@ -290,11 +290,6 @@ UEdGraphNode* UBlueprintNodeSpawner::GetTemplateNode(UEdGraph* TargetGraph, FBin
 {       
 	TRACE_CPUPROFILER_EVENT_SCOPE(UBlueprintNodeSpawner::GetTemplateNode);
 	
-	if(!DoesRequireTemplateNode())
-	{
-		return nullptr;
-	}
-
 	UEdGraphNode* TemplateNode = BlueprintNodeSpawnerImpl::GetSharedTemplateCache()->GetNodeTemplate(this, TargetGraph);
 
 	if (TemplateNode && Bindings.Num() > 0) 
diff --git a/Engine/Source/Editor/BlueprintGraph/Public/BlueprintNodeSpawner.h b/Engine/Source/Editor/BlueprintGraph/Public/BlueprintNodeSpawner.h
--- a/Engine/Source/Editor/BlueprintGraph/Public/BlueprintNodeSpawner.h
+++ b/Engine/Source/Editor/BlueprintGraph/Public/BlueprintNodeSpawner.h
@@ -295,8 +295,6 @@ private:
 	 * @return Null if it failed to spawn a node (if NodeClass is null), otherwise a newly spawned node.
 	 */
 	UE_API UEdGraphNode* SpawnEdGraphNode(TSubclassOf<UEdGraphNode> NodeClass, UEdGraph* ParentGraph, const FBindingSet& Bindings, FVector2D Location, FCustomizeNodeDelegate PostSpawnDelegate) const;
-
-	virtual bool DoesRequireTemplateNode() const { return true; }
 };
 
 /*******************************************************************************
