UE-352508 - Crash when compiling Blueprint Function containing a connected Set Variable node for a local variable named after a replicated Actor property

FKCHandler_VariableSet::Transform was not correctly filtering against local properties that happen to have the same name as a class property with CPF_Net set.

#Synced-CL 48842190

diff --git a/Engine/Source/Editor/BlueprintGraph/Private/VariableSetHandler.cpp b/Engine/Source/Editor/BlueprintGraph/Private/VariableSetHandler.cpp
--- a/Engine/Source/Editor/BlueprintGraph/Private/VariableSetHandler.cpp
+++ b/Engine/Source/Editor/BlueprintGraph/Private/VariableSetHandler.cpp
@@ -168,17 +168,25 @@ void FKCHandler_VariableSet::Compile(FKismetFunctionContext& Context, UEdGraphNo
 
 void FKCHandler_VariableSet::Transform(FKismetFunctionContext& Context, UEdGraphNode* Node)
 {
+	UK2Node_VariableSet* SetNode = Cast<UK2Node_VariableSet>(Node);
+	if (!SetNode)
+	{
+		return;
+	}
+
 	// Expands node out to include a (local) call to the RepNotify function if necessary
-	UK2Node_VariableSet* SetNotify = Cast<UK2Node_VariableSet>(Node);
 
-	// If property is HasFieldNotificationBroadcast, then the net code and broadcast will be executed in native code.
-	if (SetNotify && !SetNotify->HasFieldNotificationBroadcast())
+	const bool bIsLocalScope = SetNode->VariableReference.IsLocalScope();
+
+	// Local function properties can't be used in replication.
+	// Additionally, if HasFieldNotificationBroadcast is true, then the net code broadcast will be executed in native code.
+	if (!bIsLocalScope && !SetNode->HasFieldNotificationBroadcast())
 	{
-		if (SetNotify->ShouldFlushDormancyOnSet())
+		if (SetNode->ShouldFlushDormancyOnSet())
 		{
 			// Create CallFuncNode
 			UK2Node_CallFunction* CallFuncNode = Node->GetGraph()->CreateIntermediateNode<UK2Node_CallFunction>();
-			CallFuncNode->FunctionReference.SetExternalMember(NAME_FlushNetDormancy, AActor::StaticClass() );
+			CallFuncNode->FunctionReference.SetExternalMember(NAME_FlushNetDormancy, AActor::StaticClass());
 			CallFuncNode->AllocateDefaultPins();
 
 			// Copy self pin
@@ -199,10 +207,10 @@ void FKCHandler_VariableSet::Transform(FKismetFunctionContext& Context, UEdGraph
 			}
 		}
 
-		if (SetNotify->HasLocalRepNotify())
+		if (SetNode->HasLocalRepNotify())
 		{
 			UK2Node_CallFunction* CallFuncNode = Node->GetGraph()->CreateIntermediateNode<UK2Node_CallFunction>();
-			CallFuncNode->FunctionReference.SetExternalMember(SetNotify->GetRepNotifyName(), SetNotify->GetVariableSourceClass() );
+			CallFuncNode->FunctionReference.SetExternalMember(SetNode->GetRepNotifyName(), SetNode->GetVariableSourceClass());
 			CallFuncNode->AllocateDefaultPins();
 
 			// Copy self pin
@@ -224,16 +232,16 @@ void FKCHandler_VariableSet::Transform(FKismetFunctionContext& Context, UEdGraph
 			}
 		}
 
-		if (SetNotify->IsNetProperty())
+		if (SetNode->IsNetProperty())
 		{
 			/**
 			 * This code is for property dirty tracking.
 			 * It works by injecting in extra nodes while compiling that will call UNetPushModelHelpers::MarkPropertyDirtyFromRepIndex.
 			 * See FKCPushModelHelpers::ConstructMarkDirtyNodeForProperty for node generation.
 			 */
-			if (FProperty * Property = SetNotify->GetPropertyForVariable())
+			if (FProperty* Property = SetNode->GetPropertyForVariable())
 			{
-				if (UEdGraphNode * MarkPropertyDirtyNode = FKCPushModelHelpers::ConstructMarkDirtyNodeForProperty(Context, Property, Node->FindPinChecked(UEdGraphSchema_K2::PN_Self)))
+				if (UEdGraphNode* MarkPropertyDirtyNode = FKCPushModelHelpers::ConstructMarkDirtyNodeForProperty(Context, Property, Node->FindPinChecked(UEdGraphSchema_K2::PN_Self)))
 				{
 					// Hook up our exec pins.
 					UEdGraphPin* OldThenPin = Node->FindPinChecked(UEdGraphSchema_K2::PN_Then);
