Blueprints: make FVector2D default values work when exposed on spawn on the Construct Object BP node.

#Synced-CL 49781371

diff --git a/Engine/Source/Editor/BlueprintGraph/Classes/EdGraphSchema_K2.h b/Engine/Source/Editor/BlueprintGraph/Classes/EdGraphSchema_K2.h
--- a/Engine/Source/Editor/BlueprintGraph/Classes/EdGraphSchema_K2.h
+++ b/Engine/Source/Editor/BlueprintGraph/Classes/EdGraphSchema_K2.h
@@ -1299,6 +1299,7 @@ private:
 
 	static UE_API const UScriptStruct* VectorStruct;
 	static UE_API const UScriptStruct* Vector3fStruct;
+	static UE_API const UScriptStruct* Vector2DStruct;
 	static UE_API const UScriptStruct* RotatorStruct;
 	static UE_API const UScriptStruct* TransformStruct;
 	static UE_API const UScriptStruct* LinearColorStruct;
diff --git a/Engine/Source/Editor/BlueprintGraph/Private/EdGraphSchema_K2.cpp b/Engine/Source/Editor/BlueprintGraph/Private/EdGraphSchema_K2.cpp
--- a/Engine/Source/Editor/BlueprintGraph/Private/EdGraphSchema_K2.cpp
+++ b/Engine/Source/Editor/BlueprintGraph/Private/EdGraphSchema_K2.cpp
@@ -643,6 +643,7 @@ const int32 UEdGraphSchema_K2::AG_LevelReference = 100;
 
 const UScriptStruct* UEdGraphSchema_K2::VectorStruct = nullptr;
 const UScriptStruct* UEdGraphSchema_K2::Vector3fStruct = nullptr;
+const UScriptStruct* UEdGraphSchema_K2::Vector2DStruct = nullptr;
 const UScriptStruct* UEdGraphSchema_K2::RotatorStruct = nullptr;
 const UScriptStruct* UEdGraphSchema_K2::TransformStruct = nullptr;
 const UScriptStruct* UEdGraphSchema_K2::LinearColorStruct = nullptr;
@@ -669,6 +670,7 @@ UEdGraphSchema_K2::UEdGraphSchema_K2(const FObjectInitializer& ObjectInitializer
 	{
 		VectorStruct = TBaseStructure<FVector>::Get();
 		Vector3fStruct = TVariantStructure<FVector3f>::Get();
+		Vector2DStruct = TBaseStructure<FVector2D>::Get();
 		RotatorStruct = TBaseStructure<FRotator>::Get();
 		TransformStruct = TBaseStructure<FTransform>::Get();
 		LinearColorStruct = TBaseStructure<FLinearColor>::Get();
@@ -1419,6 +1421,7 @@ bool UEdGraphSchema_K2::PinDefaultValueIsEditable(const UEdGraphPin& InGraphPin)
 		UObject const& SubCategoryObject = *InGraphPin.PinType.PinSubCategoryObject;
 		return &SubCategoryObject == VectorStruct 
 			|| &SubCategoryObject == Vector3fStruct
+			|| &SubCategoryObject == Vector2DStruct
 			|| &SubCategoryObject == RotatorStruct
 			|| &SubCategoryObject == TransformStruct
 			|| &SubCategoryObject == LinearColorStruct
