Fall back to conservative, ordered n^2 pin matching when an obvious match is not present

[RN] Nested split pins no longer cause data loss or assertion failure when their backing schema (Struct or Function) is changed
#Synced-CL 50958978

diff --git a/Engine/Source/Editor/BlueprintGraph/Classes/K2Node.h b/Engine/Source/Editor/BlueprintGraph/Classes/K2Node.h
--- a/Engine/Source/Editor/BlueprintGraph/Classes/K2Node.h
+++ b/Engine/Source/Editor/BlueprintGraph/Classes/K2Node.h
@@ -590,6 +590,8 @@ private:
 	 */
 	bool TryInsertDefaultValueConversionNode(const UEdGraphPin& OldPin, UEdGraphPin& NewPin) const;
 
+	/** A fast pin rewiring function that does not handle split pins or orphaned pins, but safely aborts if any are found */
+	bool TryRewireOldPinsToNewPinsFast(TArray<UEdGraphPin*>& InOldPins, TArray<UEdGraphPin*>& InNewPins, TMap<UEdGraphPin*, UEdGraphPin*>* NewPinToOldPin);
 public:
 
 	/** Util to get the generated class from a node. */
diff --git a/Engine/Source/Editor/BlueprintGraph/Private/K2Node.cpp b/Engine/Source/Editor/BlueprintGraph/Private/K2Node.cpp
--- a/Engine/Source/Editor/BlueprintGraph/Private/K2Node.cpp
+++ b/Engine/Source/Editor/BlueprintGraph/Private/K2Node.cpp
@@ -1260,6 +1260,55 @@ bool UK2Node::TryInsertDefaultValueConversionNode(const UEdGraphPin& OldPin, UEd
 	return bCanConvertDefaultValue;
 }
 
+bool UK2Node::TryRewireOldPinsToNewPinsFast(TArray<UEdGraphPin*>& InOldPins, TArray<UEdGraphPin*>& InNewPins, TMap<UEdGraphPin*, UEdGraphPin*>* NewPinToOldPin)
+{
+	if(InOldPins.Num() != InNewPins.Num() )
+	{
+		return false;
+	}
+
+	for(UEdGraphPin* Pin : InOldPins)
+	{
+		if(Pin->ParentPin || Pin->bOrphanedPin)
+		{
+			return false;
+		}
+	}
+	
+	const int32 NumPins = InOldPins.Num();
+	for(int32 Idx = 0; Idx < NumPins; ++Idx)
+	{
+		UEdGraphPin* OldPin = InOldPins[Idx];
+		UEdGraphPin* NewPin = InNewPins[Idx];
+		// don't trust anything except the base implementation - customizations
+		// may do very relaxed matching (e.g. K2Node_MakeStruct):
+		ERedirectType RedirectType = UK2Node::DoPinsMatchForReconstruction(NewPin, Idx, OldPin, Idx);
+		if(RedirectType != ERedirectType_Name)
+		{
+			return false;
+		}
+	}
+
+	// We are commmitted, preconditions validated, naive, fast rewiring can be performed:
+
+	// check for mutation buried in DoPinsMatchForReconstruction:
+	checkSlow(InOldPins.Num() == InNewPins.Num() && InOldPins.Num() == NumPins);
+	// all pins match, rewire them:
+	for(int32 Idx = 0; Idx < NumPins; ++Idx)
+	{
+		UEdGraphPin* OldPin = InOldPins[Idx];
+		UEdGraphPin* NewPin = InNewPins[Idx];
+		ReconstructSinglePin(NewPin, OldPin, ERedirectType_Name);
+		checkSlow(InOldPins.Num() == InNewPins.Num() && InOldPins.Num() == NumPins);
+		if (NewPinToOldPin)
+		{
+			NewPinToOldPin->Add(NewPin, OldPin);
+		}
+	}
+
+	return true;
+}
+
 void UK2Node::EarlyValidation(FCompilerResultsLog& MessageLog) const
 {
 	UBlueprint* BP = GetBlueprint();
@@ -1337,13 +1386,15 @@ FString UK2Node::GetPinMetaData(FName InPinName, FName InKey)
 
 void UK2Node::RewireOldPinsToNewPins(TArray<UEdGraphPin*>& InOldPins, TArray<UEdGraphPin*>& InNewPins, TMap<UEdGraphPin*, UEdGraphPin*>* NewPinToOldPin)
 {
+	if(TryRewireOldPinsToNewPinsFast(InOldPins, InNewPins, NewPinToOldPin))
+	{
+		return;
+	}
+
 	TArray<UEdGraphPin*> OrphanedOldPins;
-	TArray<bool> NewPinMatched; // Tracks whether a NewPin has already been matched to an OldPin
 	TMap<UEdGraphPin*, UEdGraphPin*> MatchedPins; // Old to New
 
-	const int32 NumNewPins = InNewPins.Num();
-	NewPinMatched.AddDefaulted(NumNewPins);
-	const bool bSaveUnconnectedDefaultPins = (NumNewPins == 0 && (IsA<UK2Node_CallFunction>() || IsA<UK2Node_MacroInstance>()));
+	const bool bSaveUnconnectedDefaultPins = (InNewPins.Num() == 0 && (IsA<UK2Node_CallFunction>() || IsA<UK2Node_MacroInstance>()));
 
 	// Rewire any connection to pins that are matched by name (O(N^2) right now)
 	// NOTE: we iterate backwards through the list because ReconstructSinglePin()
@@ -1352,45 +1403,54 @@ void UK2Node::RewireOldPinsToNewPins(TArray<UEdGraphPin*>& InOldPins, TArray<UEd
 	//       we want to destroy old pins from the split children (leafs) up, so 
 	//       we do this since split child pins are ordered later in the list 
 	//       (after their parents) 
-	for (int32 OldPinIndex = InOldPins.Num()-1; OldPinIndex >= 0; --OldPinIndex)
+	const auto FindNewPin = [&](UEdGraphPin* const OldPin, const int32 OldPinIndex, TFunctionRef<ERedirectType(UEdGraphPin*, int32)> Matcher) -> bool
 	{
-		UEdGraphPin* OldPin = InOldPins[OldPinIndex];
-
-		// common case is for InOldPins and InNewPins to match, so we start searching from the current index:
 		bool bMatched = false;
-		int32 NewPinIndex = (NumNewPins ? OldPinIndex % NumNewPins : 0);
-		for (int32 NewPinCount = NumNewPins - 1; NewPinCount >= 0; --NewPinCount)
+		const int32 NumNewPins = InNewPins.Num();
+		for (int32 NewPinIndex = NumNewPins - 1; NewPinIndex >= 0; --NewPinIndex)
 		{
-			// if InNewPins grows then we may skip entries and fail to find a match or NewPinMatched will not be accurate
-			check(NumNewPins == InNewPins.Num());
-			if (!NewPinMatched[NewPinIndex])
-			{
-				UEdGraphPin* NewPin = InNewPins[NewPinIndex];
-
-				ERedirectType RedirectType = DoPinsMatchForReconstruction(NewPin, NewPinIndex, OldPin, OldPinIndex);
+			UEdGraphPin* NewPin = InNewPins[NewPinIndex];
 
-				if (RedirectType == ERedirectType_DefaultValue)
+			ERedirectType RedirectType = Matcher(NewPin, NewPinIndex);
+			if (RedirectType == ERedirectType_DefaultValue)
+			{
+				if (!TryInsertDefaultValueConversionNode(*OldPin, *NewPin))
 				{
-					if (!TryInsertDefaultValueConversionNode(*OldPin, *NewPin))
-					{
-						RedirectType = ERedirectType_None;
-					}
+					RedirectType = ERedirectType_None;
 				}
+			}
 
-				if (RedirectType != ERedirectType_None)
+			if (RedirectType != ERedirectType_None)
+			{
+				ReconstructSinglePin(NewPin, OldPin, RedirectType);
+				MatchedPins.Add(OldPin, NewPin);
+				if (NewPinToOldPin)
 				{
-					ReconstructSinglePin(NewPin, OldPin, RedirectType);
-					MatchedPins.Add(OldPin, NewPin);
-					if (NewPinToOldPin)
-					{
-						NewPinToOldPin->Add(NewPin, OldPin);
-					}
-					bMatched = true;
-					NewPinMatched[NewPinIndex] = true;
-					break;
+					NewPinToOldPin->Add(NewPin, OldPin);
 				}
+				bMatched = true;
+				break;
 			}
-			NewPinIndex = (NewPinIndex + 1) % InNewPins.Num();
+		}
+		return bMatched;
+	};
+
+	for (int32 OldPinIndex = InOldPins.Num()-1; OldPinIndex >= 0; --OldPinIndex)
+	{
+		UEdGraphPin* OldPin = InOldPins[OldPinIndex];
+
+		bool bMatched = FindNewPin(OldPin, OldPinIndex, [this, OldPin, OldPinIndex](UEdGraphPin* NewPin, int32 NewPinIndex)
+		{
+			return UK2Node::DoPinsMatchForReconstruction(NewPin, NewPinIndex, OldPin, OldPinIndex);
+		});
+
+		if(!bMatched)
+		{
+			// do proper vtable dispatch of DoPinsMatchForReconstruction in a second pass:
+			bMatched = FindNewPin(OldPin, OldPinIndex, [this, OldPin, OldPinIndex](UEdGraphPin* NewPin, int32 NewPinIndex)
+				{
+					return DoPinsMatchForReconstruction(NewPin, NewPinIndex, OldPin, OldPinIndex);
+				});
 		}
 
 		// Orphaned pins are those that existed in the OldPins array but do not in the NewPins.
@@ -1488,6 +1548,8 @@ void UK2Node::RewireOldPinsToNewPins(TArray<UEdGraphPin*>& InOldPins, TArray<UEd
 				}
 			};
 
+			ParentPin->SubPins.Remove(OrphanedPin);
+			OrphanedPin->ParentPin = nullptr;
 			InNewPins.Insert(OrphanedPin, InsertIndex);
 		}
 	}
diff --git a/Engine/Source/Editor/KismetCompiler/Private/UserDefinedStructureCompilerUtils.cpp b/Engine/Source/Editor/KismetCompiler/Private/UserDefinedStructureCompilerUtils.cpp
--- a/Engine/Source/Editor/KismetCompiler/Private/UserDefinedStructureCompilerUtils.cpp
+++ b/Engine/Source/Editor/KismetCompiler/Private/UserDefinedStructureCompilerUtils.cpp
@@ -626,10 +626,6 @@ void FUserDefinedStructureCompilerUtils::CompileStruct(UUserDefinedStruct* Struc
 		TSet<UBlueprint*> BlueprintsThatHaveBeenRecompiled;
 		FBlueprintEditorUtils::FindScriptStructsInNodes(ChangedStructsSet, [&BlueprintsThatHaveBeenRecompiled, &BlueprintsToRecompile](UBlueprint* Blueprint, UK2Node* Node)
 			{
-				// We need to recombine any nested subpins on this node, otherwise there will be an
-				// unexpected amount of pins during reconstruction. 
-				FBlueprintEditorUtils::RecombineNestedSubPins(Node);
-
 				if (Blueprint)
 				{
 					// The blueprint skeleton needs to be updated before we reconstruct the node
diff --git a/Engine/Source/Editor/UnrealEd/Private/Kismet2/ReloadUtilities.cpp b/Engine/Source/Editor/UnrealEd/Private/Kismet2/ReloadUtilities.cpp
--- a/Engine/Source/Editor/UnrealEd/Private/Kismet2/ReloadUtilities.cpp
+++ b/Engine/Source/Editor/UnrealEd/Private/Kismet2/ReloadUtilities.cpp
@@ -1127,11 +1127,6 @@ void FReload::Reinstance()
 		UBlueprint* Blueprint = KVP.Key;
 		FBlueprintUpdateInfo& Info = KVP.Value;
 
-		for (UK2Node* Node : Info.Nodes)
-		{
-			FBlueprintEditorUtils::RecombineNestedSubPins(Node);
-		}
-
 		// We must reconstruct the node first other wise some pins might not be 
 		// in a good state for the recompile
 		for (UK2Node* Node : Info.Nodes)
