UE-366986 - Float pin drag and drop should create a double-precision float

When dropping a pin on an entry or result node, it was possible for users to unintentionally create a single-precision float pin. Since we don't permit Blueprint users to create single-precision floats, we need to promote any dropped float pins to double-precision.

#Synced-CL 50676684

diff --git a/Engine/Source/Editor/BlueprintGraph/Private/EdGraphSchema_K2.cpp b/Engine/Source/Editor/BlueprintGraph/Private/EdGraphSchema_K2.cpp
--- a/Engine/Source/Editor/BlueprintGraph/Private/EdGraphSchema_K2.cpp
+++ b/Engine/Source/Editor/BlueprintGraph/Private/EdGraphSchema_K2.cpp
@@ -7719,6 +7719,15 @@ UEdGraphPin* UEdGraphSchema_K2::DropPinOnNode(UEdGraphNode* InTargetNode, const
 		TArray<UK2Node_EditablePinBase*> EditablePinNodes;
 		EditablePinNode->Modify();
 
+		FEdGraphPinType SourcePinType = InSourcePinType;
+
+		// When a user creates a new float pin, we need to ensure that it's a double.
+		// Otherwise, we're unintentionally creating a loophole where users can create single-precision floats.
+		if (SourcePinType.PinCategory == PC_Real)
+		{
+			SourcePinType.PinSubCategory = PC_Double;
+		}
+
 		if (InSourcePinDirection == EGPD_Output && Cast<UK2Node_FunctionEntry>(InTargetNode))
 		{
 			if (UK2Node_FunctionResult* ResultNode = FBlueprintEditorUtils::FindOrCreateFunctionResultNode(EditablePinNode))
@@ -7761,7 +7770,7 @@ UEdGraphPin* UEdGraphSchema_K2::DropPinOnNode(UEdGraphNode* InTargetNode, const
 		for (UK2Node_EditablePinBase* CurrentEditablePinNode : EditablePinNodes)
 		{
 			CurrentEditablePinNode->Modify();
-			UEdGraphPin* CreatedPin = CurrentEditablePinNode->CreateUserDefinedPin(InSourcePinName, InSourcePinType, (InSourcePinDirection == EGPD_Input) ? EGPD_Output : EGPD_Input);
+			UEdGraphPin* CreatedPin = CurrentEditablePinNode->CreateUserDefinedPin(InSourcePinName, SourcePinType, (InSourcePinDirection == EGPD_Input) ? EGPD_Output : EGPD_Input);
 
 			// The final ResultPin is from the node the user dragged and dropped to
 			if (EditablePinNode == CurrentEditablePinNode)
