single character typo fix in tooltip "pined" replaced with "pinned"

#rnx
#Synced-CL 51536473

diff --git a/Engine/Source/Editor/BlueprintGraph/Public/BlueprintEditorSettings.h b/Engine/Source/Editor/BlueprintGraph/Public/BlueprintEditorSettings.h
--- a/Engine/Source/Editor/BlueprintGraph/Public/BlueprintEditorSettings.h
+++ b/Engine/Source/Editor/BlueprintGraph/Public/BlueprintEditorSettings.h
@@ -134,7 +134,7 @@ public:
 	UPROPERTY(EditAnywhere, config, Category=Workflow, meta=(DisplayName="Context Menu: Expose All Sub-Component Functions"))
 	bool bExposeAllMemberComponentFunctions;
 
-	/** If enabled, then a separate section with your Blueprint favorites will be pined to the top of the context menu. */
+	/** If enabled, then a separate section with your Blueprint favorites will be pinned to the top of the context menu. */
 	UPROPERTY(EditAnywhere, config, Category=Workflow, meta=(DisplayName="Context Menu: Show Favorites Section"))
 	bool bShowContextualFavorites;
 
