SA fix, use explicit .Get() to get raw pointer out of TSubclassOf

#rnx
#Synced-CL 51516599

diff --git a/Engine/Source/Editor/BlueprintGraph/Private/K2Node_SpawnActorFromClass.cpp b/Engine/Source/Editor/BlueprintGraph/Private/K2Node_SpawnActorFromClass.cpp
--- a/Engine/Source/Editor/BlueprintGraph/Private/K2Node_SpawnActorFromClass.cpp
+++ b/Engine/Source/Editor/BlueprintGraph/Private/K2Node_SpawnActorFromClass.cpp
@@ -349,7 +349,7 @@ FText UK2Node_SpawnActorFromClass::GetNodeTitle(ENodeTitleType::Type TitleType)
 bool UK2Node_SpawnActorFromClass::IsCompatibleWithGraph(const UEdGraph* TargetGraph) const 
 {
 	UBlueprint* Blueprint = FBlueprintEditorUtils::FindBlueprintForGraph(TargetGraph);
-	const UClass* Class = Blueprint ? Blueprint->GeneratedClass : nullptr;
+	const UClass* Class = Blueprint ? Blueprint->GeneratedClass.Get() : nullptr;
 	return Super::IsCompatibleWithGraph(TargetGraph) &&
 		(	!Blueprint || 
 			(FBlueprintEditorUtils::FindUserConstructionScript(Blueprint) != TargetGraph && Class && Class->GetDefaultObject()->ImplementsGetWorld()));
