Added estimate node width. Needed for EDA comment nodes for Blueprints toolset.

#Synced-CL 52068286

diff --git a/Engine/Source/Editor/BlueprintGraph/Classes/EdGraphSchema_K2.h b/Engine/Source/Editor/BlueprintGraph/Classes/EdGraphSchema_K2.h
--- a/Engine/Source/Editor/BlueprintGraph/Classes/EdGraphSchema_K2.h
+++ b/Engine/Source/Editor/BlueprintGraph/Classes/EdGraphSchema_K2.h
@@ -1196,6 +1196,9 @@ public:
 	/** Calculates an estimated height for the specified node */
 	static UE_API float EstimateNodeHeight( UEdGraphNode* Node );
 
+	/** Calculates an estimated width for the specified node */
+	static UE_API float EstimateNodeWidth( const UEdGraphNode* Node );
+
 	/** 
 	 * Checks if the graph supports impure functions
 	 *
diff --git a/Engine/Source/Editor/BlueprintGraph/Private/EdGraphSchema_K2.cpp b/Engine/Source/Editor/BlueprintGraph/Private/EdGraphSchema_K2.cpp
--- a/Engine/Source/Editor/BlueprintGraph/Private/EdGraphSchema_K2.cpp
+++ b/Engine/Source/Editor/BlueprintGraph/Private/EdGraphSchema_K2.cpp
@@ -24,6 +24,8 @@
 #include "Kismet/GameplayStatics.h"
 #include "Engine/LevelScriptActor.h"
 #include "Components/ChildActorComponent.h"
+#include "Editor/EditorEngine.h"
+#include "Engine/Font.h"
 #include "Engine/Selection.h"
 #include "Engine/UserDefinedEnum.h"
 #include "StructUtils/UserDefinedStruct.h"
@@ -7036,6 +7038,37 @@ float UEdGraphSchema_K2::EstimateNodeHeight( UEdGraphNode* Node )
 	return HeightEstimate;
 }
 
+/**
+ * Attempts to best-guess the width of the node. This is necessary because we don't know the actual
+ * size of the node until the next Slate tick
+ *
+ * @param Node The node to guess the width of
+ * @return The estimated width of the specified node
+ */
+float UEdGraphSchema_K2::EstimateNodeWidth( const UEdGraphNode* Node )
+{
+	float WidthEstimate = 0.0f;
+
+	if ( Node != NULL )
+	{
+		const float BaseNodeWidth = Node->IsA( UK2Node_CallFunction::StaticClass() ) ? 250.0f : 200.0f;
+
+		FString NodeTitle = Node->GetNodeTitle( ENodeTitleType::FullTitle ).ToString();
+
+		if ( const UFont* Font = GetDefault<UEditorEngine>()->EditorFont )
+		{
+			WidthEstimate = FMath::Max( BaseNodeWidth, (float)Font->GetStringSize( *NodeTitle ) );
+		}
+		else
+		{
+			static const float EstimatedCharWidth = 6.0f;
+			WidthEstimate = FMath::Max( BaseNodeWidth, NodeTitle.Len() * EstimatedCharWidth );
+		}
+	}
+
+	return WidthEstimate;
+}
+
 
 bool UEdGraphSchema_K2::CollapseGatewayNode(UK2Node* InNode, UEdGraphNode* InEntryNode, UEdGraphNode* InResultNode, FKismetCompilerContext* CompilerContext, TSet<UEdGraphNode*>* OutExpandedNodes) const
 {
