Prevent transient classes that are not native from showing up in the bp menu, as if they are used references to them cannot be saved

[RN] Fields and properties associated with transient, non-native uclasses are no longer shown in the Blueprint context menu, as when placed in a graph they will not be able to resolve themselves when the asset is loaded
#Synced-CL 52432244

diff --git a/Engine/Source/Editor/BlueprintGraph/Private/BlueprintActionDatabase.cpp b/Engine/Source/Editor/BlueprintGraph/Private/BlueprintActionDatabase.cpp
--- a/Engine/Source/Editor/BlueprintGraph/Private/BlueprintActionDatabase.cpp
+++ b/Engine/Source/Editor/BlueprintGraph/Private/BlueprintActionDatabase.cpp
@@ -1486,8 +1486,11 @@ void FBlueprintActionDatabase::RefreshClassActions(UClass* const Class)
 	bool const bHiddenClass		 = Class->HasAnyClassFlags(CLASS_Hidden);
 	bool const bIsBlueprintClass = (Cast<UBlueprintGeneratedClass>(Class) != nullptr);
 	bool const bIsLevelScript	 = Class->ClassGeneratedBy && Cast<UBlueprint>(Class->ClassGeneratedBy)->BlueprintType == EBlueprintType::BPTYPE_LevelScript;
+	// transient non native classes can't be persistently referenced - carve out for blueprint classes is for skeleton classes
+	// which are transient but serve as a proxy for the persistent generated class:
+	bool const bIsTransientClass = Class->HasAnyFlags(RF_Transient) && !Class->IsNative() && !bIsBlueprintClass;
 
-	if (bOutOfDateClass || bIsLevelScript || bHiddenClass)
+	if (bOutOfDateClass || bIsLevelScript || bHiddenClass || bIsTransientClass)
 	{
 		ActionRegistry.Remove(Class);
 		return;
